Mesh Editor: Converted Mesh Editor to a plugin
- All code and content for Mesh Editor has been moved to it's own plugin
- Various fixes to extensibility of VR Editor UI

diff --git a/Engine/Plugins/Editor/MeshEditor/MeshEditor.uplugin b/Engine/Plugins/Editor/MeshEditor/MeshEditor.uplugin
new file mode 100644
--- /dev/null
+++ b/Engine/Plugins/Editor/MeshEditor/MeshEditor.uplugin
@@ -0,0 +1,25 @@
+{
+	"FileVersion" : 3,
+	"Version" : 1,
+	"VersionName" : "1.0",
+	"FriendlyName" : "Mesh Editor",
+	"Description" : "Mesh editing system",
+	"Category" : "Editor",
+	"CreatedBy" : "Epic Games, Inc.",
+	"CreatedByURL" : "http://epicgames.com",
+	"DocsURL" : "",
+	"MarketplaceURL" : "",
+	"SupportURL" : "",
+	"EnabledByDefault" : true,
+	"CanContainContent" : true,
+	"IsBetaVersion" : false,
+	"Installed" : false,
+	"Modules" :
+	[
+		{
+			"Name" : "MeshEditor",
+			"Type" : "Editor",
+			"LoadingPhase" : "Default"
+		}
+	]
+}
\ No newline at end of file
diff --git a/Engine/Plugins/Editor/MeshEditor/Source/MeshEditor/MeshEditor.Build.cs b/Engine/Plugins/Editor/MeshEditor/Source/MeshEditor/MeshEditor.Build.cs
new file mode 100644
--- /dev/null
+++ b/Engine/Plugins/Editor/MeshEditor/Source/MeshEditor/MeshEditor.Build.cs
@@ -0,0 +1,42 @@
+// Copyright 1998-2015 Epic Games, Inc. All Rights Reserved.
+
+namespace UnrealBuildTool.Rules
+{
+	public class MeshEditor : ModuleRules
+	{
+        public MeshEditor(TargetInfo Target)
+		{
+			PrivateDependencyModuleNames.AddRange(
+                new string[] {
+				    "Core",
+				    "CoreUObject",
+				    "Engine",
+                    "InputCore",
+				    "Slate",
+					"SlateCore",
+                    "EditorStyle",
+				    "UnrealEd",
+					"RenderCore",
+					"MeshEditingRuntime",
+					"ViewportInteraction",
+                    "VREditor",
+					"Projects"
+				}
+			);
+
+			PrivateIncludePathModuleNames.AddRange(
+				new string[] {
+                    "ContentBrowser",
+					"LevelEditor",
+				}
+			);
+
+			DynamicallyLoadedModuleNames.AddRange(
+				new string[] {
+                    "ContentBrowser",
+					"LevelEditor",
+				}
+			);
+		}
+	}
+}
\ No newline at end of file
diff --git a/Engine/Plugins/Editor/MeshEditor/Source/MeshEditor/MeshEditorAssetContainer.h b/Engine/Plugins/Editor/MeshEditor/Source/MeshEditor/MeshEditorAssetContainer.h
new file mode 100644
--- /dev/null
+++ b/Engine/Plugins/Editor/MeshEditor/Source/MeshEditor/MeshEditorAssetContainer.h
@@ -0,0 +1,24 @@
+// Copyright 1998-2017 Epic Games, Inc. All Rights Reserved.
+#pragma once
+
+#include "Engine/DataAsset.h"
+#include "MeshEditorAssetContainer.generated.h"
+
+/**
+ * Asset container for the mesh editor
+ */
+UCLASS()
+class MESHEDITOR_API UMeshEditorAssetContainer : public UDataAsset
+{
+    GENERATED_BODY()
+ 
+public:
+	UPROPERTY(EditAnywhere, Category = Material)
+	class UMaterial* HoveredGeometryMaterial;
+
+	UPROPERTY(EditAnywhere, Category = Material)
+	class UMaterial* HoveredFaceMaterial;
+
+	UPROPERTY(EditAnywhere, Category = Sound)
+	class USoundBase* DefaultSound;
+};
\ No newline at end of file
diff --git a/Engine/Plugins/Editor/MeshEditor/Source/MeshEditor/MeshEditorCommands.cpp b/Engine/Plugins/Editor/MeshEditor/Source/MeshEditor/MeshEditorCommands.cpp
new file mode 100644
--- /dev/null
+++ b/Engine/Plugins/Editor/MeshEditor/Source/MeshEditor/MeshEditorCommands.cpp
@@ -0,0 +1,64 @@
+// Copyright 1998-2015 Epic Games, Inc. All Rights Reserved.
+
+#include "MeshEditorCommands.h"
+
+#define LOCTEXT_NAMESPACE "MeshEditorCommands"
+
+void FMeshEditorCommonCommands::RegisterCommands()
+{
+	UI_COMMAND(DeleteMeshElement, "Delete", "Delete selected mesh elements, including polygons partly defined by selected elements.", EUserInterfaceActionType::Button, FInputChord(EKeys::Delete));
+	UI_COMMAND(AddSubdivisionLevel, "Add Subdivision Level", "Increases the number of subdivision levels for the selected mesh.", EUserInterfaceActionType::Button, FInputChord(EKeys::Equals));
+	UI_COMMAND(RemoveSubdivisionLevel, "Remove Subdivision Level", "Decreases the number of subdivision levels for the selected mesh.", EUserInterfaceActionType::Button, FInputChord(EKeys::Hyphen));
+	UI_COMMAND(ShowVertexNormals, "Show Vertex Normals", "Toggles debug rendering of vertex normals.", EUserInterfaceActionType::ToggleButton, FInputChord(EKeys::V));
+	UI_COMMAND(MarqueeSelectVertices, "Marquee Select Vertices", "Selects vertices inside the selection box.", EUserInterfaceActionType::None, FInputChord());
+	UI_COMMAND(MarqueeSelectEdges, "Marquee Select Edges", "Selects edges inside the selection box.", EUserInterfaceActionType::None, FInputChord());
+	UI_COMMAND(MarqueeSelectPolygons, "Marquee Select Polygons", "Selects polygons inside the selection box.", EUserInterfaceActionType::None, FInputChord());
+	UI_COMMAND(DrawVertices, "Draw Vertices", "Allows vertices to be freely drawn to create a new polygon.", EUserInterfaceActionType::RadioButton, FInputChord());
+	UI_COMMAND(FrameSelectedElements, "Frame Selected Elements", "Moves the viewport camera to frame the currently selected elements.", EUserInterfaceActionType::None, FInputChord(EKeys::F));
+	UI_COMMAND(SetVertexSelectionMode, "Set Vertex Selection Mode", "Sets the selection mode so that only vertices will be selected.", EUserInterfaceActionType::None, FInputChord(EKeys::One));
+	UI_COMMAND(SetEdgeSelectionMode, "Set Edge Selection Mode", "Sets the selection mode so that only edges will be selected.", EUserInterfaceActionType::None, FInputChord(EKeys::Two));
+	UI_COMMAND(SetPolygonSelectionMode, "Set Polygon Selection Mode", "Sets the selection mode so that only polygons will be selected.", EUserInterfaceActionType::None, FInputChord(EKeys::Three));
+	UI_COMMAND(SetAnySelectionMode, "Set Any Selection Mode", "Sets the selection mode so that any element type may be selected.", EUserInterfaceActionType::None, FInputChord(EKeys::Four));
+	UI_COMMAND(QuadrangulateMesh, "Quadrangulate Mesh", "Quadrangulates the selected mesh.", EUserInterfaceActionType::Button, FInputChord());
+}
+
+void FMeshEditorVertexCommands::RegisterCommands()
+{
+	UI_COMMAND(MoveVertex, "Move Vertex Mode", "Set the primary action to move vertices.", EUserInterfaceActionType::RadioButton, FInputChord(EKeys::F1));
+	UI_COMMAND(ExtendVertex, "Extend Vertex Mode", "Set the primary action to extend vertices.", EUserInterfaceActionType::RadioButton, FInputChord(EKeys::F2));
+	UI_COMMAND(EditVertexCornerSharpness, "Edit Vertex Corner Sharpness Mode", "Set the primary action to edit the vertex's subdivision corner sharpness amount.", EUserInterfaceActionType::RadioButton, FInputChord(EKeys::F6));
+
+	UI_COMMAND(RemoveVertex, "Remove Vertex", "Remove the selected vertex if possible.", EUserInterfaceActionType::Button, FInputChord(EKeys::BackSpace));
+	UI_COMMAND(WeldVertices, "Weld Vertices", "Weld the selected vertices, keeping the first selected vertex.", EUserInterfaceActionType::Button, FInputChord());
+}
+
+void FMeshEditorEdgeCommands::RegisterCommands()
+{
+	UI_COMMAND(MoveEdge, "Move Edge Mode", "Set the primary action to move edges.", EUserInterfaceActionType::RadioButton, FInputChord(EKeys::F1));
+	UI_COMMAND(SplitEdge, "Split Edge Mode", "Set the primary action to split edges.", EUserInterfaceActionType::RadioButton, FInputChord(EKeys::F2));
+	UI_COMMAND(SplitEdgeAndDragVertex, "Split Edge and Drag Vertex Mode", "Set the primary action to split edges and drag vertices.", EUserInterfaceActionType::RadioButton, FInputChord(EKeys::F3));
+	UI_COMMAND(InsertEdgeLoop, "Insert Edge Loop Mode", "Set the primary action to insert edge loops.", EUserInterfaceActionType::RadioButton, FInputChord(EKeys::F4));
+	UI_COMMAND(ExtendEdge, "Extend Edge Mode", "Set the primary action to extend edges.", EUserInterfaceActionType::RadioButton, FInputChord(EKeys::F5));
+	UI_COMMAND(EditEdgeCreaseSharpness, "Edit Edge Crease Sharpness Mode", "Set the primary action to edit the edge's crease sharpness.", EUserInterfaceActionType::RadioButton, FInputChord(EKeys::F6));
+
+	UI_COMMAND(RemoveEdge, "Remove Edge", "Remove the selected edge if possible.", EUserInterfaceActionType::Button, FInputChord(EKeys::BackSpace));
+	UI_COMMAND(SoftenEdge, "Soften Edge", "Make selected edge soft.", EUserInterfaceActionType::Button, FInputChord(EKeys::H, EModifierKey::Shift));
+	UI_COMMAND(HardenEdge, "Harden Edge", "Make selected edge hard.", EUserInterfaceActionType::Button, FInputChord(EKeys::H));
+	UI_COMMAND(SelectEdgeLoop, "Select Edge Loop", "Select the edge loops which contain the selected edges.", EUserInterfaceActionType::Button, FInputChord(EKeys::Two, EModifierKey::Shift));
+}
+
+void FMeshEditorPolygonCommands::RegisterCommands()
+{
+	UI_COMMAND(MovePolygon, "Move Polygon Mode", "Set the primary action to move polygons.", EUserInterfaceActionType::RadioButton, FInputChord(EKeys::F1));
+	UI_COMMAND(ExtrudePolygon, "Extrude Polygon Mode", "Set the primary action to extrude polygons.", EUserInterfaceActionType::RadioButton, FInputChord(EKeys::F2));
+	// UI_COMMAND(FreelyExtrudePolygon, "Freely Extrude Polygon Mode", "Set the primary action to freely extrude polygons.", EUserInterfaceActionType::RadioButton, FInputChord(EKeys::F3));
+	UI_COMMAND(InsetPolygon, "Inset Polygon Mode", "Set the primary action to inset polygons.", EUserInterfaceActionType::RadioButton, FInputChord(EKeys::F3));
+	UI_COMMAND(BevelPolygon, "Bevel Polygon Mode", "Set the primary action to bevel polygons.", EUserInterfaceActionType::RadioButton, FInputChord(EKeys::F4));
+
+	UI_COMMAND(FlipPolygon, "Flip Polygon", "Flip the currently selected polygons.", EUserInterfaceActionType::Button, FInputChord(EKeys::F, EModifierKey::Shift));
+	UI_COMMAND(TriangulatePolygon, "Triangulate Polygon", "Triangulate the currently selected polygons.", EUserInterfaceActionType::Button, FInputChord(EKeys::T));
+	UI_COMMAND(TessellatePolygon, "Tessellate Selected Polygons", "Tessellate selected polygons into smaller polygons.", EUserInterfaceActionType::Button, FInputChord());
+	UI_COMMAND(AssignMaterial, "Assign Material", "Assigns the highlighted material in the Content Browser to the currently selected polygons.", EUserInterfaceActionType::Button, FInputChord(EKeys::M));
+}
+
+#undef LOCTEXT_NAMESPACE
\ No newline at end of file
