Mesh Editor extensibility (phase 1)
- This is the first change toward making it easy to add new polygon modeling features in external plugins (or through scripting!)
- Added a new 'PolygonModeling' module to the MeshEditor plugin and moved the TessellatePolygon feature there
- New UMeshEditorPolygonCommand ABC class (and friends) for polygon editing features to inherit from
- MeshEditor now has a very minimal public API
- Experimenting with using multiple protected 'contract interfaces' for FMeshEditorMode to reduce public API size
- Some initial clean-ups to how selection/deselection works (reduce boilerplate.)
- Removed a bad use of 'friend' for FMeshElementViewportTransformable
- Moved a few functions around to either make them publicly accessible or to hide them
- Slate: Added support for registering UICommandInfos without a 'this' pointer

diff --git a/Engine/Plugins/Editor/MeshEditor/MeshEditor.uplugin b/Engine/Plugins/Editor/MeshEditor/MeshEditor.uplugin
--- a/Engine/Plugins/Editor/MeshEditor/MeshEditor.uplugin
+++ b/Engine/Plugins/Editor/MeshEditor/MeshEditor.uplugin
@@ -20,6 +20,11 @@
 			"Name" : "MeshEditor",
 			"Type" : "Editor",
 			"LoadingPhase" : "Default"
+		},
+		{
+			"Name" : "PolygonModeling",
+			"Type" : "Editor",
+			"LoadingPhase" : "Default"
 		}
 	]
 }
\ No newline at end of file
diff --git a/Engine/Plugins/Editor/MeshEditor/Source/MeshEditor/IMeshEditorModeUIContract.h b/Engine/Plugins/Editor/MeshEditor/Source/MeshEditor/IMeshEditorModeUIContract.h
new file mode 100644
--- /dev/null
+++ b/Engine/Plugins/Editor/MeshEditor/Source/MeshEditor/IMeshEditorModeUIContract.h
@@ -0,0 +1,35 @@
+// Copyright 1998-2015 Epic Games, Inc. All Rights Reserved.
+
+#pragma once
+
+#include "EditableMeshTypes.h"
+
+
+class IMeshEditorModeUIContract
+{
+
+public:
+
+	/** Returns the current selection mode we're in */
+	virtual EEditableMeshElementType GetMeshElementSelectionMode() const = 0;
+	
+	/** Sets the mesh element selection mode to use */
+	virtual void SetMeshElementSelectionMode( EEditableMeshElementType ElementType ) = 0;
+
+	/** Returns the type of elements that are selected right now, or Invalid if nothing is selected */
+	virtual EEditableMeshElementType GetSelectedMeshElementType() const = 0;
+
+	virtual const TArray<TTuple<TSharedPtr<FUICommandInfo>, FUIAction>>& GetVertexActions() const = 0;
+	virtual const TArray<TTuple<TSharedPtr<FUICommandInfo>, FUIAction>>& GetEdgeActions() const = 0;
+	virtual const TArray<TTuple<TSharedPtr<FUICommandInfo>, FUIAction>>& GetPolygonActions() const = 0;
+
+	virtual bool IsEditingPerInstance() const = 0;
+	virtual void SetEditingPerInstance( bool bPerInstance ) = 0;
+
+	/** Propagates instance changes to the static mesh asset */
+	virtual void PropagateInstanceChanges() = 0;
+
+	/** Whether there are instance changes which can be propagated */
+	virtual bool CanPropagateInstanceChanges() const = 0;
+
+};
diff --git a/Engine/Plugins/Editor/MeshEditor/Source/MeshEditor/MeshEditorCommands.cpp b/Engine/Plugins/Editor/MeshEditor/Source/MeshEditor/MeshEditorCommands.cpp
--- a/Engine/Plugins/Editor/MeshEditor/Source/MeshEditor/MeshEditorCommands.cpp
+++ b/Engine/Plugins/Editor/MeshEditor/Source/MeshEditor/MeshEditorCommands.cpp
@@ -1,9 +1,20 @@
 // Copyright 1998-2015 Epic Games, Inc. All Rights Reserved.
 
 #include "MeshEditorCommands.h"
+#include "MeshEditorStyle.h"
 
 #define LOCTEXT_NAMESPACE "MeshEditorCommands"
 
+FMeshEditorCommonCommands::FMeshEditorCommonCommands() 
+	: TCommands<FMeshEditorCommonCommands>(
+		"MeshEditorCommon",
+		LOCTEXT( "MeshEditorGeneral", "Mesh Editor Common" ),
+		"MainFrame",
+		FMeshEditorStyle::GetStyleSetName()	)
+{
+}
+
+
 void FMeshEditorCommonCommands::RegisterCommands()
 {
 	UI_COMMAND(DeleteMeshElement, "Delete", "Delete selected mesh elements, including polygons partly defined by selected elements.", EUserInterfaceActionType::Button, FInputChord(EKeys::Delete));
@@ -22,6 +33,16 @@ void FMeshEditorCommonCommands::RegisterCommands()
 	UI_COMMAND(QuadrangulateMesh, "Quadrangulate Mesh", "Quadrangulates the selected mesh.", EUserInterfaceActionType::Button, FInputChord());
 }
 
+FMeshEditorVertexCommands::FMeshEditorVertexCommands() 
+	: TCommands<FMeshEditorVertexCommands>(
+		"MeshEditorVertex",
+		LOCTEXT("MeshEditorVertex", "Mesh Editor Vertex"),
+		"MeshEditorCommon",
+		FMeshEditorStyle::GetStyleSetName())
+{
+}
+
+
 void FMeshEditorVertexCommands::RegisterCommands()
 {
 	UI_COMMAND(MoveVertex, "Move Vertex Mode", "Set the primary action to move vertices.", EUserInterfaceActionType::RadioButton, FInputChord(EKeys::F1));
@@ -32,6 +53,16 @@ void FMeshEditorVertexCommands::RegisterCommands()
 	UI_COMMAND(WeldVertices, "Weld Vertices", "Weld the selected vertices, keeping the first selected vertex.", EUserInterfaceActionType::Button, FInputChord());
 }
 
+FMeshEditorEdgeCommands::FMeshEditorEdgeCommands()
+	: TCommands<FMeshEditorEdgeCommands>(
+		"MeshEditorEdge",
+		LOCTEXT("MeshEditorEdge", "Mesh Editor Edge"),
+		"MeshEditorCommon",
+		FMeshEditorStyle::GetStyleSetName() )
+{
+}
+
+
 void FMeshEditorEdgeCommands::RegisterCommands()
 {
 	UI_COMMAND(MoveEdge, "Move Edge Mode", "Set the primary action to move edges.", EUserInterfaceActionType::RadioButton, FInputChord(EKeys::F1));
@@ -47,6 +78,16 @@ void FMeshEditorEdgeCommands::RegisterCommands()
 	UI_COMMAND(SelectEdgeLoop, "Select Edge Loop", "Select the edge loops which contain the selected edges.", EUserInterfaceActionType::Button, FInputChord(EKeys::Two, EModifierKey::Shift));
 }
 
+FMeshEditorPolygonCommands::FMeshEditorPolygonCommands() 
+	: TCommands<FMeshEditorPolygonCommands>(
+		"MeshEditorPolygon",
+		LOCTEXT("MeshEditorPolygon", "Mesh Editor Polygon"),
+		"MeshEditorCommon",
+		FMeshEditorStyle::GetStyleSetName() )
+{
+}
+
+
 void FMeshEditorPolygonCommands::RegisterCommands()
 {
 	UI_COMMAND(MovePolygon, "Move Polygon Mode", "Set the primary action to move polygons.", EUserInterfaceActionType::RadioButton, FInputChord(EKeys::F1));
@@ -57,8 +98,16 @@ void FMeshEditorPolygonCommands::RegisterCommands()
 
 	UI_COMMAND(FlipPolygon, "Flip Polygon", "Flip the currently selected polygons.", EUserInterfaceActionType::Button, FInputChord(EKeys::F, EModifierKey::Shift));
 	UI_COMMAND(TriangulatePolygon, "Triangulate Polygon", "Triangulate the currently selected polygons.", EUserInterfaceActionType::Button, FInputChord(EKeys::T));
-	UI_COMMAND(TessellatePolygon, "Tessellate Selected Polygons", "Tessellate selected polygons into smaller polygons.", EUserInterfaceActionType::Button, FInputChord());
 	UI_COMMAND(AssignMaterial, "Assign Material", "Assigns the highlighted material in the Content Browser to the currently selected polygons.", EUserInterfaceActionType::Button, FInputChord(EKeys::M));
+
+	for( TObjectIterator<UMeshEditorPolygonCommand> PolygonCommandCDOIter( RF_NoFlags ); PolygonCommandCDOIter; ++PolygonCommandCDOIter )
+	{
+		UMeshEditorPolygonCommand* PolygonCommandCDO = *PolygonCommandCDOIter;
+		if( !( PolygonCommandCDO->GetClass()->GetClassFlags() & CLASS_Abstract ) )
+		{
+			PolygonCommandCDO->RegisterUICommand( this );
+		}
+	}
 }
 
 #undef LOCTEXT_NAMESPACE
\ No newline at end of file
diff --git a/Engine/Plugins/Editor/MeshEditor/Source/MeshEditor/MeshEditorCommands.h b/Engine/Plugins/Editor/MeshEditor/Source/MeshEditor/MeshEditorCommands.h
deleted file mode 100644
diff --git a/Engine/Plugins/Editor/MeshEditor/Source/MeshEditor/MeshEditorMode.cpp b/Engine/Plugins/Editor/MeshEditor/Source/MeshEditor/MeshEditorMode.cpp
--- a/Engine/Plugins/Editor/MeshEditor/Source/MeshEditor/MeshEditorMode.cpp
+++ b/Engine/Plugins/Editor/MeshEditor/Source/MeshEditor/MeshEditorMode.cpp
@@ -2,6 +2,7 @@
 
 #include "MeshEditorMode.h"
 #include "MeshEditorCommands.h"
+#include "MeshEditorStyle.h"
 #include "MeshEditorModeToolkit.h"
 #include "EditableMesh.h"
 #include "EditableMeshFactory.h"
@@ -34,6 +35,8 @@
 #include "VREditorWorldInteraction.h"
 #include "MeshEditorAssetContainer.h"
 #include "Kismet/GameplayStatics.h"
+#include "UndoHelpers.h"
+
 
 #define LOCTEXT_NAMESPACE "MeshEditorMode"
 
@@ -559,9 +562,22 @@ void FMeshEditorMode::BindCommands()
 
 	RegisterPolygonCommand( MeshEditorPolygonActions.FlipPolygon, FExecuteAction::CreateLambda( [this] { FlipSelectedPolygons(); } ) );
 	RegisterPolygonCommand( MeshEditorPolygonActions.TriangulatePolygon, FExecuteAction::CreateLambda( [this] { TriangulateSelectedPolygons(); } ) );
-	RegisterPolygonCommand( MeshEditorPolygonActions.TessellatePolygon, FExecuteAction::CreateLambda( [ this ] { TessellateSelectedPolygons(); } ) );
 	RegisterPolygonCommand( MeshEditorPolygonActions.AssignMaterial, FExecuteAction::CreateLambda( [this] { AssignSelectedMaterialToSelectedPolygons(); } ) );
 
+	for( TObjectIterator<UMeshEditorPolygonCommand> PolygonCommandCDOIter( RF_NoFlags ); PolygonCommandCDOIter; ++PolygonCommandCDOIter )
+	{
+		UMeshEditorPolygonCommand* PolygonCommandCDO = *PolygonCommandCDOIter;
+		if( !( PolygonCommandCDO->GetClass()->GetClassFlags() & CLASS_Abstract ) )
+		{
+			RegisterPolygonCommand( PolygonCommandCDO->GetUICommandInfo(), FExecuteAction::CreateLambda( [this, PolygonCommandCDO]
+			{
+				bool bWasSuccessful = false;
+				PolygonCommandCDO->Execute( *this, /* Out */ bWasSuccessful );
+			} ) );
+		}
+	}
+
+
 	// Bind common actions
 	CommonCommands = MakeShared<FUICommandList>();
 	for( const TTuple<TSharedPtr<FUICommandInfo>, FUIAction>& CommonAction : CommonActions )
@@ -694,7 +710,8 @@ void FMeshEditorMode::Enter()
 	// Add toolkit
 	if( !Toolkit.IsValid() )
 	{
-		Toolkit = MakeShared< FMeshEditorModeToolkit >( *this );
+		IMeshEditorModeUIContract* UIContract = this;
+		Toolkit = MakeShared< FMeshEditorModeToolkit >( *UIContract );
 		Toolkit->Init( Owner->GetToolkitHost() );
 	}
 }
@@ -713,7 +730,7 @@ void FMeshEditorMode::Exit()
 	{
 		const FScopedTransaction Transaction( LOCTEXT( "UndoDeselectingAllMeshElements", "Deselect All Elements" ) );
 
-		StoreUndo( MeshEditorModeProxyObject, FDeselectAllMeshElementsChange( MoveTemp( FDeselectAllMeshElementsChangeInput() ) ).Execute( MeshEditorModeProxyObject ) );
+		UndoHelpers::StoreUndo( MeshEditorModeProxyObject, FDeselectAllMeshElementsChange( MoveTemp( FDeselectAllMeshElementsChangeInput() ) ).Execute( MeshEditorModeProxyObject ) );
 	}
 
 	FToolkitManager::Get().CloseToolkit( Toolkit.ToSharedRef() );
@@ -1036,7 +1053,7 @@ void FMeshEditorMode::CommitEditableMeshIfNecessary( UEditableMesh* EditableMesh
 				ChangeInput.MeshElementsToSelect.Add( NewMeshElement );
 			}
 		}
-		StoreUndo( MeshEditorModeProxyObject, FSelectOrDeselectMeshElementsChange( MoveTemp( ChangeInput ) ).Execute( MeshEditorModeProxyObject ) );
+		UndoHelpers::StoreUndo( MeshEditorModeProxyObject, FSelectOrDeselectMeshElementsChange( MoveTemp( ChangeInput ) ).Execute( MeshEditorModeProxyObject ) );
 
 		FixUpMeshElements( SelectedVertices );
 		FixUpMeshElements( SelectedEdges );
@@ -1132,6 +1149,30 @@ const UMeshEditorAssetContainer& FMeshEditorMode::GetAssetContainer() const
 	return *AssetContainer;
 }
 
+void FMeshEditorMode::SelectMeshElements( const TArray<FMeshElement>& MeshElementsToSelect )
+{
+	FCompoundChangeInput CompoundRevertInput;
+	FSelectOrDeselectMeshElementsChangeInput ChangeInput;
+	ChangeInput.MeshElementsToSelect = MeshElementsToSelect;
+
+	CompoundRevertInput.Subchanges.Add( MoveTemp( FSelectOrDeselectMeshElementsChange( MoveTemp( ChangeInput ) ).Execute( MeshEditorModeProxyObject ) ) );
+
+	UndoHelpers::StoreUndo( MeshEditorModeProxyObject, MakeUnique<FCompoundChange>( MoveTemp( CompoundRevertInput ) ) );
+}
+
+
+void FMeshEditorMode::DeselectMeshElements( const TArray<FMeshElement>& MeshElementsToDeselect )
+{
+	FCompoundChangeInput CompoundRevertInput;
+	FSelectOrDeselectMeshElementsChangeInput ChangeInput;
+	ChangeInput.MeshElementsToDeselect = MeshElementsToDeselect;
+
+	CompoundRevertInput.Subchanges.Add( MoveTemp( FSelectOrDeselectMeshElementsChange( MoveTemp( ChangeInput ) ).Execute( MeshEditorModeProxyObject ) ) );
+
+	UndoHelpers::StoreUndo( MeshEditorModeProxyObject, MakeUnique<FCompoundChange>( MoveTemp( CompoundRevertInput ) ) );
+}
+
+
 void FMeshEditorMode::DeselectMeshElements( const TMap<UEditableMesh*, TArray<FMeshElement>>& MeshElementsToDeselect )
 {
 	FCompoundChangeInput CompoundRevertInput;
@@ -1146,7 +1187,7 @@ void FMeshEditorMode::DeselectMeshElements( const TMap<UEditableMesh*, TArray<FM
 
 	CompoundRevertInput.Subchanges.Add( MoveTemp( FSelectOrDeselectMeshElementsChange( MoveTemp( ChangeInput ) ).Execute( MeshEditorModeProxyObject ) ) );
 
-	StoreUndo( MeshEditorModeProxyObject, MakeUnique<FCompoundChange>( MoveTemp( CompoundRevertInput ) ) );
+	UndoHelpers::StoreUndo( MeshEditorModeProxyObject, MakeUnique<FCompoundChange>( MoveTemp( CompoundRevertInput ) ) );
 }
 
 
@@ -1220,101 +1261,11 @@ bool FMeshEditorMode::DeleteSelectedMeshElement()
 
 		EditableMesh->EndModification();
 
-		StoreUndo( EditableMesh, EditableMesh->MakeUndo() );
+		UndoHelpers::StoreUndo( EditableMesh, EditableMesh->MakeUndo() );
 	}
 
 	// Make sure we're not still hovering over a polygon we deleted
-	StoreUndo( MeshEditorModeProxyObject, MoveTemp( ClearInvalidSelectedElements_Internal() ) );
-
-	return true;
-}
-
-
-bool FMeshEditorMode::TessellateSelectedPolygons()
-{
-	if( ActiveAction != EMeshEditAction::None )
-	{
-		return false;
-	}
-
-	static TMap<UEditableMesh*, TArray<FMeshElement>> SelectedMeshesAndPolygons;
-	GetSelectedMeshesAndPolygons( SelectedMeshesAndPolygons );
-	if( SelectedMeshesAndPolygons.Num() == 0 )
-	{
-		return false;
-	}
-
-	FScopedTransaction Transaction( LOCTEXT( "UndoTessellatePolygon", "Tessellate Polygon" ) );
-
-	CommitSelectedMeshes();
-
-	// Refresh selection (committing may have created a new mesh instance)
-	GetSelectedMeshesAndPolygons( SelectedMeshesAndPolygons );
-
-	// Deselect the mesh elements before we delete them.  This will make sure they become selected again after undo.
-	DeselectMeshElements( SelectedMeshesAndPolygons );
-
-	static TArray<FMeshElement> MeshElementsToSelect;
-	MeshElementsToSelect.Reset();
-
-	for( const auto& SelectedMeshAndPolygons : SelectedMeshesAndPolygons )
-	{
-		UEditableMesh* EditableMesh = SelectedMeshAndPolygons.Key;
-		const TArray<FMeshElement>& PolygonElements = SelectedMeshAndPolygons.Value;
-
-		static TArray<FPolygonRef> PolygonsToTessellate;
-		PolygonsToTessellate.Reset( PolygonElements.Num() );
-
-		for( const auto& PolygonElement : PolygonElements )
-		{
-			const FPolygonRef PolygonRef( PolygonElement.ElementAddress.SectionID, FPolygonID( PolygonElement.ElementAddress.ElementID ) );
-			PolygonsToTessellate.Add( PolygonRef );
-		}
-
-
-		EditableMesh->StartModification( EMeshModificationType::Final, EMeshTopologyChange::TopologyChange );
-
-		// @todo mesheditor: Expose as configurable parameter
-		const ETriangleTessellationMode TriangleTessellationMode = ETriangleTessellationMode::FourTriangles;
-
-		static TArray<FPolygonRef> NewPolygonRefs;
-		EditableMesh->TessellatePolygons( PolygonsToTessellate, TriangleTessellationMode, /* Out */ NewPolygonRefs );
-
-		// Select the new polygons
-		for( const FPolygonRef NewPolygonRef : NewPolygonRefs )
-		{
-			FMeshElement NewPolygonMeshElement;
-			{
-				NewPolygonMeshElement.Component = PolygonElements[0].Component;
-				NewPolygonMeshElement.ElementAddress = PolygonElements[0].ElementAddress;
-				NewPolygonMeshElement.ElementAddress.ElementType = EEditableMeshElementType::Polygon;
-				NewPolygonMeshElement.ElementAddress.SectionID = NewPolygonRef.SectionID;
-				NewPolygonMeshElement.ElementAddress.ElementID = NewPolygonRef.PolygonID;
-			}
-
-			MeshElementsToSelect.Add( NewPolygonMeshElement );
-		}
-
-		EditableMesh->EndModification();
-
-		StoreUndo( EditableMesh, EditableMesh->MakeUndo() );
-	}
-
-	{
-		FCompoundChangeInput CompoundRevertInput;
-
-		// Make sure we're not still hovering over a polygon we removed
-		CompoundRevertInput.Subchanges.Add( MoveTemp( ClearInvalidSelectedElements_Internal() ) );
-
-		// Select the polygon leftover after removing the edge
-		FSelectOrDeselectMeshElementsChangeInput ChangeInput;
-		ChangeInput.MeshElementsToSelect = MeshElementsToSelect;
-		CompoundRevertInput.Subchanges.Add( MoveTemp( FSelectOrDeselectMeshElementsChange( MoveTemp( ChangeInput ) ).Execute( MeshEditorModeProxyObject ) ) );
-
-		StoreUndo( MeshEditorModeProxyObject, MakeUnique<FCompoundChange>( MoveTemp( CompoundRevertInput ) ) );
-	}
-
-	StoreUndo( MeshEditorModeProxyObject, MoveTemp( ClearInvalidSelectedElements_Internal() ) );
+	ClearInvalidSelectedElements();
 
 	return true;
 }
@@ -1367,7 +1318,7 @@ void FMeshEditorMode::AddOrRemoveSubdivisionLevel( const bool bShouldAdd )
 
 			EditableMesh->EndModification();
 
-			StoreUndo( EditableMesh, EditableMesh->MakeUndo() );
+			UndoHelpers::StoreUndo( EditableMesh, EditableMesh->MakeUndo() );
 		}
 	}
 }
@@ -1404,7 +1355,7 @@ void FMeshEditorMode::QuadrangulateMesh()
 
 		FCompoundChangeInput CompoundRevertInput;
 		CompoundRevertInput.Subchanges.Add( MoveTemp( FDeselectAllMeshElementsChange( FDeselectAllMeshElementsChangeInput() ).Execute( MeshEditorModeProxyObject ) ) );
-		StoreUndo( MeshEditorModeProxyObject, MakeUnique<FCompoundChange>( MoveTemp( CompoundRevertInput ) ) );
+		UndoHelpers::StoreUndo( MeshEditorModeProxyObject, MakeUnique<FCompoundChange>( MoveTemp( CompoundRevertInput ) ) );
 
 		for( UEditableMesh* EditableMesh : SelectedMeshes )
 		{
@@ -1413,7 +1364,7 @@ void FMeshEditorMode::QuadrangulateMesh()
 			EditableMesh->QuadrangulateMesh( NewPolygonRefs );
 			EditableMesh->EndModification();
 
-			StoreUndo( EditableMesh, EditableMesh->MakeUndo() );
+			UndoHelpers::StoreUndo( EditableMesh, EditableMesh->MakeUndo() );
 		}
 	}
 }
@@ -1582,21 +1533,21 @@ bool FMeshEditorMode::RemoveSelectedEdges()
 
 		EditableMesh->EndModification();
 
-		StoreUndo( EditableMesh, EditableMesh->MakeUndo() );
+		UndoHelpers::StoreUndo( EditableMesh, EditableMesh->MakeUndo() );
 	}
 
 	{
-		FCompoundChangeInput CompoundRevertInput;
-
 		// Make sure we're not still hovering over a polygon we removed
-		CompoundRevertInput.Subchanges.Add( MoveTemp( ClearInvalidSelectedElements_Internal() ) );
+		ClearInvalidSelectedElements();
+
+		FCompoundChangeInput CompoundRevertInput;
 
 		// Select the polygon leftover after removing the edge
 		FSelectOrDeselectMeshElementsChangeInput ChangeInput;
 		ChangeInput.MeshElementsToSelect = MeshElementsToSelect;
 		CompoundRevertInput.Subchanges.Add( MoveTemp( FSelectOrDeselectMeshElementsChange( MoveTemp( ChangeInput ) ).Execute( MeshEditorModeProxyObject ) ) );
 
-		StoreUndo( MeshEditorModeProxyObject, MakeUnique<FCompoundChange>( MoveTemp( CompoundRevertInput ) ) );
+		UndoHelpers::StoreUndo( MeshEditorModeProxyObject, MakeUnique<FCompoundChange>( MoveTemp( CompoundRevertInput ) ) );
 	}
 
 	return true;
@@ -1656,7 +1607,7 @@ bool FMeshEditorMode::SelectEdgeLoops()
 		ChangeInput.MeshElementsToSelect = MeshElementsToSelect;
 		CompoundRevertInput.Subchanges.Add( MoveTemp( FSelectOrDeselectMeshElementsChange( MoveTemp( ChangeInput ) ).Execute( MeshEditorModeProxyObject ) ) );
 
-		StoreUndo( MeshEditorModeProxyObject, MakeUnique<FCompoundChange>( MoveTemp( CompoundRevertInput ) ) );
+		UndoHelpers::StoreUndo( MeshEditorModeProxyObject, MakeUnique<FCompoundChange>( MoveTemp( CompoundRevertInput ) ) );
 	}
 
 	return true;
@@ -1742,21 +1693,21 @@ bool FMeshEditorMode::RemoveSelectedVertices()
 
 		EditableMesh->EndModification();
 
-		StoreUndo( EditableMesh, EditableMesh->MakeUndo() );
+		UndoHelpers::StoreUndo( EditableMesh, EditableMesh->MakeUndo() );
 	}
 
 	{
-		FCompoundChangeInput CompoundRevertInput;
-
 		// Make sure we're not still hovering over a polygon we removed
-		CompoundRevertInput.Subchanges.Add( MoveTemp( ClearInvalidSelectedElements_Internal() ) );
+		ClearInvalidSelectedElements();
+
+		FCompoundChangeInput CompoundRevertInput;
 
 		// Select the edge that the vertex we removed was a member of
 		FSelectOrDeselectMeshElementsChangeInput ChangeInput;
 		ChangeInput.MeshElementsToSelect = MeshElementsToSelect;
 		CompoundRevertInput.Subchanges.Add( MoveTemp( FSelectOrDeselectMeshElementsChange( MoveTemp( ChangeInput ) ).Execute( MeshEditorModeProxyObject ) ) );
 
-		StoreUndo( MeshEditorModeProxyObject, MakeUnique<FCompoundChange>( MoveTemp( CompoundRevertInput ) ) );
+		UndoHelpers::StoreUndo( MeshEditorModeProxyObject, MakeUnique<FCompoundChange>( MoveTemp( CompoundRevertInput ) ) );
 	}
 
 	return true;
@@ -1836,20 +1787,21 @@ bool FMeshEditorMode::WeldSelectedVertices()
 
 		EditableMesh->EndModification();
 
-		StoreUndo( EditableMesh, EditableMesh->MakeUndo() );
+		UndoHelpers::StoreUndo( EditableMesh, EditableMesh->MakeUndo() );
 	}
 
 	{
-		FCompoundChangeInput CompoundRevertInput;
+		// Make sure we're not still hovering over a polygon we removed
+		ClearInvalidSelectedElements();
 
-		CompoundRevertInput.Subchanges.Add( MoveTemp( ClearInvalidSelectedElements_Internal() ) );
+		FCompoundChangeInput CompoundRevertInput;
 
 		// Select the vertex that was welded.
 		FSelectOrDeselectMeshElementsChangeInput ChangeInput;
 		ChangeInput.MeshElementsToSelect = MeshElementsToSelect;
 		CompoundRevertInput.Subchanges.Add( MoveTemp( FSelectOrDeselectMeshElementsChange( MoveTemp( ChangeInput ) ).Execute( MeshEditorModeProxyObject ) ) );
 
-		StoreUndo( MeshEditorModeProxyObject, MakeUnique<FCompoundChange>( MoveTemp( CompoundRevertInput ) ) );
+		UndoHelpers::StoreUndo( MeshEditorModeProxyObject, MakeUnique<FCompoundChange>( MoveTemp( CompoundRevertInput ) ) );
 	}
 
 	return true;
@@ -1899,7 +1851,7 @@ bool FMeshEditorMode::FlipSelectedPolygons()
 
 		EditableMesh->EndModification();
 
-		StoreUndo( EditableMesh, EditableMesh->MakeUndo() );
+		UndoHelpers::StoreUndo( EditableMesh, EditableMesh->MakeUndo() );
 	}
 
 	return true;
@@ -1980,21 +1932,21 @@ bool FMeshEditorMode::TriangulateSelectedPolygons()
 
 		EditableMesh->EndModification();
 
-		StoreUndo( EditableMesh, EditableMesh->MakeUndo() );
+		UndoHelpers::StoreUndo( EditableMesh, EditableMesh->MakeUndo() );
 	}
 
 	// Select the newly-created triangles
 	{
-		FCompoundChangeInput CompoundRevertInput;
-
 		// Make sure we're not still hovering over a polygon we removed
-		CompoundRevertInput.Subchanges.Add( MoveTemp( ClearInvalidSelectedElements_Internal() ) );
+		ClearInvalidSelectedElements();
+
+		FCompoundChangeInput CompoundRevertInput;
 
 		FSelectOrDeselectMeshElementsChangeInput ChangeInput;
 		ChangeInput.MeshElementsToSelect = MeshElementsToSelect;
 		CompoundRevertInput.Subchanges.Add( MoveTemp( FSelectOrDeselectMeshElementsChange( MoveTemp( ChangeInput ) ).Execute( MeshEditorModeProxyObject ) ) );
 
-		StoreUndo( MeshEditorModeProxyObject, MakeUnique<FCompoundChange>( MoveTemp( CompoundRevertInput ) ) );
+		UndoHelpers::StoreUndo( MeshEditorModeProxyObject, MakeUnique<FCompoundChange>( MoveTemp( CompoundRevertInput ) ) );
 	}
 
 
@@ -2092,21 +2044,21 @@ bool FMeshEditorMode::AssignMaterialToSelectedPolygons( UMaterialInterface* Sele
 			}
 			EditableMesh->EndModification();
 
-			StoreUndo( EditableMesh, EditableMesh->MakeUndo() );
+			UndoHelpers::StoreUndo( EditableMesh, EditableMesh->MakeUndo() );
 		}
 
 		// Select the newly-created triangles
 		{
-			FCompoundChangeInput CompoundRevertInput;
-
 			// Make sure we're not still hovering over a polygon we removed
-			CompoundRevertInput.Subchanges.Add( MoveTemp( ClearInvalidSelectedElements_Internal() ) );
+			ClearInvalidSelectedElements();
+
+			FCompoundChangeInput CompoundRevertInput;
 
 			FSelectOrDeselectMeshElementsChangeInput ChangeInput;
 			ChangeInput.MeshElementsToSelect = MeshElementsToSelect;
 			CompoundRevertInput.Subchanges.Add( MoveTemp( FSelectOrDeselectMeshElementsChange( MoveTemp( ChangeInput ) ).Execute( MeshEditorModeProxyObject ) ) );
 
-			StoreUndo( MeshEditorModeProxyObject, MakeUnique<FCompoundChange>( MoveTemp( CompoundRevertInput ) ) );
+			UndoHelpers::StoreUndo( MeshEditorModeProxyObject, MakeUnique<FCompoundChange>( MoveTemp( CompoundRevertInput ) ) );
 		}
 
 		return true;
@@ -2164,7 +2116,7 @@ bool FMeshEditorMode::MakeSelectedEdgesHardOrSoft(const bool bMakeEdgesHard)
 
 		EditableMesh->EndModification();
 
-		StoreUndo( EditableMesh, EditableMesh->MakeUndo() );
+		UndoHelpers::StoreUndo( EditableMesh, EditableMesh->MakeUndo() );
 	}
 
 	return true;
@@ -3118,7 +3070,7 @@ void FMeshEditorMode::OnViewportInteractionInputUnhandled( FEditorViewportClient
 			{
 				const FScopedTransaction Transaction( LOCTEXT( "UndoDeselectingAllMeshElements", "Deselect All Elements" ) );
 
-				StoreUndo( MeshEditorModeProxyObject, FDeselectAllMeshElementsChange( MoveTemp( FDeselectAllMeshElementsChangeInput() ) ).Execute( MeshEditorModeProxyObject ) );
+				UndoHelpers::StoreUndo( MeshEditorModeProxyObject, FDeselectAllMeshElementsChange( MoveTemp( FDeselectAllMeshElementsChangeInput() ) ).Execute( MeshEditorModeProxyObject ) );
 			}
 		}
 	}
@@ -3332,7 +3284,7 @@ void FMeshEditorMode::UpdateActiveAction( const bool bIsActionFinishing )
 	};
 
 
-	check( !IsUndoSystemAvailable() || GEditor->IsTransactionActive() );
+	check( !UndoHelpers::IsUndoSystemAvailable() || GEditor->IsTransactionActive() );
 
 	static TSet<UEditableMesh*> MeshesBeingModified;
 	MeshesBeingModified.Reset();
@@ -4397,14 +4349,10 @@ void FMeshEditorMode::UpdateActiveAction( const bool bIsActionFinishing )
 		{
 			CompoundRevertInput.Subchanges.Add( MoveTemp( FDeselectAllMeshElementsChange( FDeselectAllMeshElementsChangeInput() ).Execute( MeshEditorModeProxyObject ) ) );
 		}
-
-		// Make sure nothing is still selected that was deleted
+		else
 		{
-			TUniquePtr<FChange> ClearInvalidRevertChange = MoveTemp( ClearInvalidSelectedElements_Internal() );
-			if( ClearInvalidRevertChange.IsValid() )
-			{
-				CompoundRevertInput.Subchanges.Add( MoveTemp( ClearInvalidRevertChange ) );
-			}
+			// Make sure nothing is still selected that was deleted
+			ClearInvalidSelectedElements();
 		}
 
 		if( MeshElementsToSelect.Num() > 0 )
@@ -4587,7 +4535,7 @@ void FMeshEditorMode::UpdateActiveAction( const bool bIsActionFinishing )
 				TUniquePtr<FChange>& RevertChange = UpdateRevertChange.Get<1>();
 				check( RevertChange.IsValid() );
 
-				StoreUndo( Object, MoveTemp( RevertChange ) );
+				UndoHelpers::StoreUndo( Object, MoveTemp( RevertChange ) );
 			}
 		}
 		else
@@ -5144,7 +5092,7 @@ void FMeshEditorMode::SetMeshElementSelectionMode( EEditableMeshElementType Elem
 	const FScopedTransaction Transaction( LOCTEXT( "ChangeMeshElementSelectionMode", "Change Mesh Element Selection Mode" ) );
 	FSetElementSelectionModeChangeInput ChangeInput;
 	ChangeInput.Mode = ElementType;
-	StoreUndo( MeshEditorModeProxyObject, FSetElementSelectionModeChange( MoveTemp( ChangeInput ) ).Execute( MeshEditorModeProxyObject ) );
+	UndoHelpers::StoreUndo( MeshEditorModeProxyObject, FSetElementSelectionModeChange( MoveTemp( ChangeInput ) ).Execute( MeshEditorModeProxyObject ) );
 }
 
 
@@ -5207,6 +5155,14 @@ TUniquePtr<FChange> FMeshEditorMode::ClearInvalidSelectedElements_Internal()
 }
 
 
+void FMeshEditorMode::ClearInvalidSelectedElements()
+{
+	// Make sure we're not still hovering over a polygon we removed
+	TUniquePtr<FChange> RevertChange = ClearInvalidSelectedElements_Internal();
+	UndoHelpers::StoreUndo( MeshEditorModeProxyObject, MoveTemp( RevertChange ) );
+}
+
+
 void FMeshEditorMode::ClearInvalidHoveredElements()
 {
 	for( FMeshEditorInteractorData& MeshEditorInteractorData : MeshEditorInteractorDatas )
@@ -5513,7 +5469,7 @@ void FMeshEditorMode::OnViewportInteractionInputAction( FEditorViewportClient& V
 
 						FSelectOrDeselectMeshElementsChangeInput ChangeInput;
 						ChangeInput.MeshElementsToDeselect.Add( SelectedMeshElements[ AlreadySelectedMeshElement ] );
-						StoreUndo( MeshEditorModeProxyObject, FSelectOrDeselectMeshElementsChange( MoveTemp( ChangeInput ) ).Execute( MeshEditorModeProxyObject ) );
+						UndoHelpers::StoreUndo( MeshEditorModeProxyObject, FSelectOrDeselectMeshElementsChange( MoveTemp( ChangeInput ) ).Execute( MeshEditorModeProxyObject ) );
 					}
 					else if( MeshElementSelectionMode == EEditableMeshElementType::Any || MeshElementSelectionMode == MeshEditorInteractorData.HoveredMeshElement.ElementAddress.ElementType )
 					{
@@ -5548,7 +5504,7 @@ void FMeshEditorMode::OnViewportInteractionInputAction( FEditorViewportClient& V
 						{
 							// If select by painting is disabled, add a transaction immediately
 							const FScopedTransaction Transaction( LOCTEXT( "SelectElement", "Select Element" ) );
-							StoreUndo( MeshEditorModeProxyObject, MoveTemp( RevertChange ) );
+							UndoHelpers::StoreUndo( MeshEditorModeProxyObject, MoveTemp( RevertChange ) );
 						}
 					}
 				}
@@ -5597,7 +5553,7 @@ void FMeshEditorMode::OnViewportInteractionInputAction( FEditorViewportClient& V
 							// canceled our transaction while the mouse was down.
 							if( GUndo != nullptr )
 							{
-								StoreUndo( MeshEditorModeProxyObject, MakeUnique<FCompoundChange>( MoveTemp( *SelectingByPaintingRevertChangeInput.Release() ) ) );
+								UndoHelpers::StoreUndo( MeshEditorModeProxyObject, MakeUnique<FCompoundChange>( MoveTemp( *SelectingByPaintingRevertChangeInput.Release() ) ) );
 							}
 						}
 						SelectingByPaintingRevertChangeInput.Reset();
@@ -5643,7 +5599,7 @@ void FMeshEditorMode::FinishAction()
 	// @todo mesheditor: Make sure this is called before Undo is invoked (PreEditUndo!), otherwise the previous action will be undone instead of the active one
 
 	check( ActiveAction != EMeshEditAction::None );
-	check( !IsUndoSystemAvailable() || GEditor->IsTransactionActive() );	// Someone must have started a transaction! (It might not have been us though.)
+	check( !UndoHelpers::IsUndoSystemAvailable() || GEditor->IsTransactionActive() );	// Someone must have started a transaction! (It might not have been us though.)
 
 	const bool bIsActionFinishing = true;
 
@@ -5685,25 +5641,6 @@ void FMeshEditorMode::FinishAction()
 }
 
 
-bool FMeshEditorMode::IsUndoSystemAvailable()
-{
-	return GUndo != nullptr;
-}
-
-
-void FMeshEditorMode::StoreUndo( UObject* Object, TUniquePtr<FChange> UndoChange )
-{
-	// Did you forget to use an FScopedTransaction?  If GUndo was null, then most likely we forgot to wrap this call within an editor transaction.
-	// The only exception is in Simulate mode, where Undo is not allowed.
-	check( IsUndoSystemAvailable() || GEditor == nullptr || GEditor->bIsSimulatingInEditor );
-
-	if( IsUndoSystemAvailable() )
-	{
-		GUndo->StoreUndo( Object, MoveTemp( UndoChange ) );
-	}	
-}
-
-
 void FMeshEditorMode::PostUndo()
 {
 	// Make sure we're not hovering over anything that was deleted after an undo action
@@ -6043,7 +5980,7 @@ void FMeshEditorMode::PerformMarqueeSelect( EEditableMeshElementType ElementType
 			break;
 	}
 
-	StoreUndo( MeshEditorModeProxyObject, FSelectOrDeselectMeshElementsChange( MoveTemp( ChangeInput ) ).Execute( MeshEditorModeProxyObject ) );
+	UndoHelpers::StoreUndo( MeshEditorModeProxyObject, FSelectOrDeselectMeshElementsChange( MoveTemp( ChangeInput ) ).Execute( MeshEditorModeProxyObject ) );
 }
 
 
@@ -6169,7 +6106,7 @@ void FMeshEditorMode::OnActorSelectionChanged( const TArray<UObject*>& NewSelect
 
 			if( ChangeInput.MeshElementsToDeselect.Num() > 0 || ChangeInput.MeshElementsToSelect.Num() > 0 )
 			{
-				StoreUndo( MeshEditorModeProxyObject, FDeselectAllMeshElementsChange( MoveTemp( FDeselectAllMeshElementsChangeInput() ) ).Execute( MeshEditorModeProxyObject ) );
+				UndoHelpers::StoreUndo( MeshEditorModeProxyObject, FDeselectAllMeshElementsChange( MoveTemp( FDeselectAllMeshElementsChangeInput() ) ).Execute( MeshEditorModeProxyObject ) );
 			}
 		}
 	}
diff --git a/Engine/Plugins/Editor/MeshEditor/Source/MeshEditor/MeshEditorMode.h b/Engine/Plugins/Editor/MeshEditor/Source/MeshEditor/MeshEditorMode.h
--- a/Engine/Plugins/Editor/MeshEditor/Source/MeshEditor/MeshEditorMode.h
+++ b/Engine/Plugins/Editor/MeshEditor/Source/MeshEditor/MeshEditorMode.h
@@ -9,6 +9,8 @@
 #include "ArrayView.h"
 #include "LevelEditorViewport.h"	// For FTrackingTransaction
 #include "UnrealEdMisc.h"
+#include "IMeshEditorModeUIContract.h"
+#include "IMeshEditorModeEditingContract.h"
 #include "MeshEditorMode.generated.h"
 
 
@@ -24,10 +26,11 @@ public:
 };
 
 
+
 /**
  * Mesh Editor Mode.  Extends editor viewports with the ability to edit meshes
  */
-class FMeshEditorMode : public FEdMode
+class FMeshEditorMode : public FEdMode, protected IMeshEditorModeUIContract, protected IMeshEditorModeEditingContract
 {
 public:
 
@@ -96,41 +99,10 @@ public:
 	/** Cleans up this mode, called when the editor is shutting down */
 	virtual ~FMeshEditorMode();
 
-	// FEdMode interface
-	virtual void Enter() override;
-	virtual void Exit() override;
-	virtual void Tick( FEditorViewportClient* ViewportClient, float DeltaTime ) override;
-	virtual bool InputKey( FEditorViewportClient* ViewportClient, FViewport* Viewport, FKey Key, EInputEvent Event ) override;
-	virtual bool InputAxis( FEditorViewportClient* ViewportClient, FViewport* Viewport, int32 ControllerId, FKey Key, float Delta, float DeltaTime ) override;
-	virtual bool InputDelta( FEditorViewportClient* ViewportClient, FViewport* Viewport, FVector& Drag, FRotator& Rotation, FVector& Scale ) override;
-	virtual bool IsCompatibleWith( FEditorModeID OtherModeID ) const override;
-	virtual void AddReferencedObjects( FReferenceCollector& Collector ) override;
-	virtual void Render( const FSceneView* SceneView, FViewport* Viewport, FPrimitiveDrawInterface* PDI ) override;	
-	virtual void PostUndo() override;
-	virtual bool FrustumSelect( const FConvexVolume& InFrustum, FEditorViewportClient* InViewportClient, bool InSelect = true ) override;
-	virtual bool ShouldDrawWidget() const override;
-
-	/** Returns the current selection mode we're in */
-	EEditableMeshElementType GetMeshElementSelectionMode() const 
-	{ 
-		return MeshElementSelectionMode; 
-	}
-	
-	/** Sets the mesh element selection mode to use */
-	void SetMeshElementSelectionMode( EEditableMeshElementType ElementType );
-
-	/** Returns the type of elements that are selected right now, or Invalid if nothing is selected */
-	EEditableMeshElementType GetSelectedMeshElementType() const;
-
-	/** Returns whether the specified element type is selected */
-	bool IsMeshElementTypeSelected( EEditableMeshElementType ElementType ) const { return GetSelectedMeshElementType() == ElementType; }
-
-	/** Returns whether either the specified element type is selected, or we're in the selection mode for that element type */
-	bool IsMeshElementTypeSelectedOrIsActiveSelectionMode( EEditableMeshElementType ElementType ) const { return GetMeshElementSelectionMode() == ElementType || GetSelectedMeshElementType() == ElementType; }
+	/** Gets an editable mesh from our cache of editable meshes for the specified sub-mesh address, or tries to create and cache a new
+	editable mesh if we haven't seen this sub-mesh address before.  Can return nullptr if no mesh was possible for that address. */
+	UEditableMesh* FindOrCreateEditableMesh( class UPrimitiveComponent& Component, const FEditableMeshSubMeshAddress& SubMeshAddress );
 
-	const TArray<TTuple<TSharedPtr<FUICommandInfo>, FUIAction>>& GetVertexActions() const { return VertexActions; }
-	const TArray<TTuple<TSharedPtr<FUICommandInfo>, FUIAction>>& GetEdgeActions() const { return EdgeActions; }
-	const TArray<TTuple<TSharedPtr<FUICommandInfo>, FUIAction>>& GetPolygonActions() const { return PolygonActions; }
 
 	/** Checks to see that the mesh element actually exists in the mesh */
 	inline static bool IsElementIDValid( const FMeshElement& MeshElement, const UEditableMesh* EditableMesh )
@@ -158,23 +130,68 @@ public:
 		return bIsValid;
 	}
 
-	bool IsEditingPerInstance() const { return bPerInstanceEdits; }
-	void SetEditingPerInstance( bool bPerInstance ) { bPerInstanceEdits = bPerInstance; }
 
-	/** Propagates instance changes to the static mesh asset */
-	void PropagateInstanceChanges();
+protected:
+
+	// FEdMode interface
+	virtual void Enter() override;
+	virtual void Exit() override;
+	virtual void Tick( FEditorViewportClient* ViewportClient, float DeltaTime ) override;
+	virtual bool InputKey( FEditorViewportClient* ViewportClient, FViewport* Viewport, FKey Key, EInputEvent Event ) override;
+	virtual bool InputAxis( FEditorViewportClient* ViewportClient, FViewport* Viewport, int32 ControllerId, FKey Key, float Delta, float DeltaTime ) override;
+	virtual bool InputDelta( FEditorViewportClient* ViewportClient, FViewport* Viewport, FVector& Drag, FRotator& Rotation, FVector& Scale ) override;
+	virtual bool IsCompatibleWith( FEditorModeID OtherModeID ) const override;
+	virtual void AddReferencedObjects( FReferenceCollector& Collector ) override;
+	virtual void Render( const FSceneView* SceneView, FViewport* Viewport, FPrimitiveDrawInterface* PDI ) override;
+	virtual void PostUndo() override;
+	virtual bool FrustumSelect( const FConvexVolume& InFrustum, FEditorViewportClient* InViewportClient, bool InSelect = true ) override;
+	virtual bool ShouldDrawWidget() const override;
+
+	// IMeshEditorModeUIContract interface
+	virtual EEditableMeshElementType GetMeshElementSelectionMode() const override
+	{ 
+		return MeshElementSelectionMode; 
+	}
+	virtual void SetMeshElementSelectionMode( EEditableMeshElementType ElementType ) override;
+	virtual EEditableMeshElementType GetSelectedMeshElementType() const override;
+	virtual const TArray<TTuple<TSharedPtr<FUICommandInfo>, FUIAction>>& GetVertexActions() const override { return VertexActions; }
+	virtual const TArray<TTuple<TSharedPtr<FUICommandInfo>, FUIAction>>& GetEdgeActions() const override { return EdgeActions; }
+	virtual const TArray<TTuple<TSharedPtr<FUICommandInfo>, FUIAction>>& GetPolygonActions() const override { return PolygonActions; }
+	virtual bool IsEditingPerInstance() const override { return bPerInstanceEdits; }
+	virtual void SetEditingPerInstance( bool bPerInstance ) override { bPerInstanceEdits = bPerInstance; }
+	virtual void PropagateInstanceChanges() override;
+	virtual bool CanPropagateInstanceChanges() const override;
+
+	// IMeshEditorModeEditingContract interface
+	virtual EMeshEditAction GetActiveAction() const	override { return ActiveAction; }
+	virtual void CommitSelectedMeshes() override;
+	virtual void GetSelectedMeshesAndVertices( TMap<UEditableMesh*, TArray<FMeshElement>>& OutMeshesAndVertices ) override 
+	{
+		GetSelectedMeshesAndElements( EEditableMeshElementType::Vertex, OutMeshesAndVertices ); 
+	}
+	virtual void GetSelectedMeshesAndEdges( TMap<UEditableMesh*, TArray<FMeshElement>>& OutMeshesAndEdges ) override 
+	{
+		GetSelectedMeshesAndElements( EEditableMeshElementType::Edge, OutMeshesAndEdges ); 
+	}
+	virtual void GetSelectedMeshesAndPolygons( TMap<UEditableMesh*, TArray<FMeshElement>>& OutMeshesAndPolygons ) override 
+	{
+		GetSelectedMeshesAndElements( EEditableMeshElementType::Polygon, OutMeshesAndPolygons ); 
+	}
+	virtual void SelectMeshElements( const TArray<FMeshElement>& MeshElementsToSelect ) override;
+	virtual void DeselectMeshElements( const TArray<FMeshElement>& MeshElementsToDeselect ) override;
+	virtual void DeselectMeshElements( const TMap<UEditableMesh*, TArray<FMeshElement>>& MeshElementsToDeselect ) override;
+	virtual void ClearInvalidSelectedElements() override;
+
 
-	/** Whether there are instance changes which can be propagated */
-	bool CanPropagateInstanceChanges() const;
 
 	/** Gets the container of all the assets used in the mesh editor */
 	const class UMeshEditorAssetContainer& GetAssetContainer() const;
 
-protected:
+	/** Returns whether the specified element type is selected */
+	bool IsMeshElementTypeSelected( EEditableMeshElementType ElementType ) const { return GetSelectedMeshElementType() == ElementType; }
 
-	/** Gets an editable mesh from our cache of editable meshes for the specified sub-mesh address, or tries to create and cache a new 
-	    editable mesh if we haven't seen this sub-mesh address before.  Can return nullptr if no mesh was possible for that address. */
-	UEditableMesh* FindOrCreateEditableMesh( class UPrimitiveComponent& Component, const FEditableMeshSubMeshAddress& SubMeshAddress );
+	/** Returns whether either the specified element type is selected, or we're in the selection mode for that element type */
+	bool IsMeshElementTypeSelectedOrIsActiveSelectionMode( EEditableMeshElementType ElementType ) const { return GetMeshElementSelectionMode() == ElementType || GetSelectedMeshElementType() == ElementType; }
 
 	/** Fills the specified dynamic mesh builder with primitives to render a mesh vertex */
 	void AddVertexToDynamicMesh( const UEditableMesh& EditableMesh, const FTransform& CameraToWorld, const FMatrix& ComponentToWorldMatrix, const FVertexID VertexID, const FColor ColorAndOpacity, const float SizeBias, const bool bApplyDepthBias, class FDynamicMeshBuilder& MeshBuilder );
@@ -224,66 +241,12 @@ protected:
 	/** Returns the index of an element in the selection set, or INDEX_NONE if its not selected */
 	int32 GetSelectedMeshElementIndex( const FMeshElement& MeshElement ) const;
 
-	/** Clears hover and selection on mesh elements that may no longer be valid.  You'll want to call this if you change the mesh topology */
-	TUniquePtr<FChange> ClearInvalidSelectedElements_Internal();
-
 	/** Clears hover on mesh elements that may no longer be valid.  You'll want to call this if you change the mesh topology */
 	void ClearInvalidHoveredElements();
 
 	/** Updates the current view location, from either the viewport client or the VR interface, whichever is in use */
 	void UpdateCameraToWorldTransform( const FEditorViewportClient& ViewportClient );
 
-	enum class EMeshEditAction
-	{
-		/** Nothing going on right now */
-		None,
-
-		/** Selecting mesh elements by 'painting' over multiple elements */
-		SelectByPainting,
-
-		/** Moving elements using a transform gizmo */
-		MoveUsingGizmo,
-
-		/** Moving selected mesh elements (vertices, edges or polygons) */
-		Move,
-
-		/** Split an edge by inserting a vertex.  You can drag to preview where the vertex will be inserted. */
-		SplitEdge,
-
-		/** Splits an edge by inserting a new vertex, then immediately starts dragging that vertex */
-		SplitEdgeAndDragVertex,
-
-		/** Insert an edge loop */
-		InsertEdgeLoop,
-
-		/** Extrude polygon by making a copy of it and allowing you to shift it along the polygon normal axis */
-		ExtrudePolygon,
-
-		/** Extrude polygon by making a copy of it and allowing you to move it around freely */
-		FreelyExtrudePolygon,
-
-		/** Inset polygon by replacing it with a new polygon that is bordered by polygons of a specific relative size */
-		InsetPolygon,
-
-		/** Bevel polygons by adding angled bordering polygons of a specific relative size */
-		BevelPolygon,
-
-		/** Extend an edge by making a copy of it and allowing you to move it around */
-		ExtendEdge,
-
-		/** Extend a vertex by making a copy of it, creating new polygons to join the geometry together */
-		ExtendVertex,
-
-		/** For subdivision meshes, edits how sharp a vertex corner is by dragging in space */
-		EditVertexCornerSharpness,
-
-		/** For subdivision meshes, edits how sharp an edge crease is by dragging in space */
-		EditEdgeCreaseSharpness,
-
-		/** Freehand vertex drawing */
-		DrawVertices,
-	};
-
 	/** Begins an action */
 	void StartAction( const EMeshEditAction NewAction, class UViewportInteractor* ActionInteractor, const bool bActionNeedsHoverLocation, const FText& UndoText );
 
@@ -291,12 +254,6 @@ protected:
 	    also be called when the user begins a new action while inertia is still influencing the active action */
 	void FinishAction();
 
-	/** @return Returns true if the undo system is available right now.  When in Simulate Mode, we can't store undo states or use undo/redo features */
-	static bool IsUndoSystemAvailable();
-
-	/** Saves undo state, if possible (e.g., not in Simulate mode.) */
-	static void StoreUndo( UObject* Object, TUniquePtr<FChange> UndoChange );
-
 	/** Binds UI commands to actions for the mesh editor */
 	void BindCommands();
 
@@ -313,18 +270,14 @@ protected:
 	/** Return the CommandList pertinent to the currently selected element type, or nullptr if nothing is selected */
 	const FUICommandList* GetCommandListForSelectedElementType() const;
 
+	TUniquePtr<FChange> ClearInvalidSelectedElements_Internal();
+
 	/** Commits the mesh instance for the given component */
 	void CommitEditableMeshIfNecessary( UEditableMesh* EditableMesh, UPrimitiveComponent* Component );
 
-	/** Commits all selected meshes */
-	void CommitSelectedMeshes();
-
 	/** Deletes selected polygons, or polygons partly defined by selected elements; returns whether successful */
 	bool DeleteSelectedMeshElement();
 
-	/** Tessellates selected polygons into smaller polygons; returns whether it did anything */
-	bool TessellateSelectedPolygons();
-
 	/** Adds or removes a subdivision level for selected meshes */
 	void AddOrRemoveSubdivisionLevel( const bool bShouldAdd );
 
@@ -368,16 +321,9 @@ protected:
 	    new (empty) data will be created for it on demand */
 	FMeshEditorInteractorData& GetMeshEditorInteractorData( UViewportInteractor* ViewportInteractor );
 
-	/** Deselects the mesh elements specified */
-	void DeselectMeshElements( const TMap<UEditableMesh*, TArray<FMeshElement>>& MeshElementsToDeselect );
-
 	/** Helper function that returns a map keying an editable mesh with its selected elements */
 	void GetSelectedMeshesAndElements( EEditableMeshElementType ElementType, TMap<UEditableMesh*, TArray<FMeshElement>>& OutMeshesAndElements );
 
-	void GetSelectedMeshesAndVertices( TMap<UEditableMesh*, TArray<FMeshElement>>& OutMeshesAndVertices ) { GetSelectedMeshesAndElements( EEditableMeshElementType::Vertex, OutMeshesAndVertices ); }
-	void GetSelectedMeshesAndEdges( TMap<UEditableMesh*, TArray<FMeshElement>>& OutMeshesAndEdges ) { GetSelectedMeshesAndElements( EEditableMeshElementType::Edge, OutMeshesAndEdges ); }
-	void GetSelectedMeshesAndPolygons( TMap<UEditableMesh*, TArray<FMeshElement>>& OutMeshesAndPolygons ) { GetSelectedMeshesAndElements( EEditableMeshElementType::Polygon, OutMeshesAndPolygons ); }
-
 	/** Helper function that converts selected vertices into a map of meshes to their respective selected vertex elements and vertex IDs */
 	void GetSelectedMeshesAndVertices( TMap< UEditableMesh*, TArray< TTuple< FMeshElement, FVertexID > > >& OutMeshesAndVertices );
 
@@ -677,7 +623,7 @@ protected:
 	EMeshTopologyChange PreviewTopologyChange;
 
 	/** Proxy UObject to pass to the undo system when performing interactions that affect state of the mode itself, 
-	    such as the selection set.  We need this because the UE4 undo system requries a UObject, but we're an FEdMode. */
+	    such as the selection set.  We need this because the UE4 undo system requires a UObject, but we're an FEdMode. */
 	UMeshEditorModeProxyObject* MeshEditorModeProxyObject;
 
 	/** When selecting by painting (EMeshEditAction::SelectByPainting), this is the compound change that can be applied to
@@ -711,7 +657,5 @@ protected:
 
 	/** Holds all the assets for the mesh editor */
 	class UMeshEditorAssetContainer* AssetContainer;
-
-	friend class FMeshElementViewportTransformable;
 };
 
diff --git a/Engine/Plugins/Editor/MeshEditor/Source/MeshEditor/MeshEditorModeToolkit.cpp b/Engine/Plugins/Editor/MeshEditor/Source/MeshEditor/MeshEditorModeToolkit.cpp
--- a/Engine/Plugins/Editor/MeshEditor/Source/MeshEditor/MeshEditorModeToolkit.cpp
+++ b/Engine/Plugins/Editor/MeshEditor/Source/MeshEditor/MeshEditorModeToolkit.cpp
@@ -1,7 +1,7 @@
 // Copyright 1998-2016 Epic Games, Inc. All Rights Reserved.
 
 #include "MeshEditorModeToolkit.h"
-#include "MeshEditorMode.h"
+#include "IMeshEditorModeUIContract.h"
 #include "MeshEditorCommands.h"
 #include "MeshEditorStyle.h"
 #include "SWidgetSwitcher.h"
@@ -119,7 +119,7 @@ public:
 	SLATE_END_ARGS()
 
 public:
-	void Construct( const FArguments& InArgs, FMeshEditorMode& MeshEditorMode, EEditableMeshElementType ElementType, const FText& Label )
+	void Construct( const FArguments& InArgs, IMeshEditorModeUIContract& MeshEditorMode, EEditableMeshElementType ElementType, const FText& Label )
 	{
 		ChildSlot
 		[
@@ -149,39 +149,36 @@ public:
 };
 
 
-void SMeshEditorModeControls::Construct( const FArguments& InArgs )
+void SMeshEditorModeControls::Construct( const FArguments& InArgs, IMeshEditorModeUIContract& MeshEditorMode )
 {
-	FMeshEditorMode* MeshEditorMode = GetMeshEditorMode();
-	check( MeshEditorMode );
-
 	// Uses the widget switcher widget so only the widget in the slot which corresponds to the selected mesh element type will be shown
 	TSharedRef<SWidgetSwitcher> WidgetSwitcher = SNew( SWidgetSwitcher )
-		.WidgetIndex_Lambda( [MeshEditorMode]() -> int32
+		.WidgetIndex_Lambda( [&MeshEditorMode]() -> int32
 			{
-				if( MeshEditorMode->GetMeshElementSelectionMode() != EEditableMeshElementType::Any )
+				if( MeshEditorMode.GetMeshElementSelectionMode() != EEditableMeshElementType::Any )
 				{
-					return static_cast<int32>( MeshEditorMode->GetMeshElementSelectionMode() );
+					return static_cast<int32>( MeshEditorMode.GetMeshElementSelectionMode() );
 				}
 				else
 				{
-					return static_cast<int32>( MeshEditorMode->GetSelectedMeshElementType() );
+					return static_cast<int32>( MeshEditorMode.GetSelectedMeshElementType() );
 				}
 			} 
 		);
 
 	WidgetSwitcher->AddSlot( static_cast<int32>( EEditableMeshElementType::Vertex ) )
 	[
-		SNew( SMeshEditorModeControlWidget, MeshEditorMode->GetVertexActions() )
+		SNew( SMeshEditorModeControlWidget, MeshEditorMode.GetVertexActions() )
 	];
 
 	WidgetSwitcher->AddSlot( static_cast<int32>( EEditableMeshElementType::Edge ) )
 	[
-		SNew( SMeshEditorModeControlWidget, MeshEditorMode->GetEdgeActions() )
+		SNew( SMeshEditorModeControlWidget, MeshEditorMode.GetEdgeActions() )
 	];
 
 	WidgetSwitcher->AddSlot( static_cast<int32>( EEditableMeshElementType::Polygon ) )
 	[
-		SNew( SMeshEditorModeControlWidget, MeshEditorMode->GetPolygonActions() )
+		SNew( SMeshEditorModeControlWidget, MeshEditorMode.GetPolygonActions() )
 	];
 
 	WidgetSwitcher->AddSlot( static_cast<int32>( EEditableMeshElementType::Invalid ) )
@@ -218,25 +215,25 @@ void SMeshEditorModeControls::Construct( const FArguments& InArgs )
 				.FillWidth( 1 )
 				.Padding( 2 )
 				[
-					SNew( SMeshEditorSelectionModeWidget, *MeshEditorMode, EEditableMeshElementType::Vertex, LOCTEXT( "Vertices", "Vertices" ) )
+					SNew( SMeshEditorSelectionModeWidget, MeshEditorMode, EEditableMeshElementType::Vertex, LOCTEXT( "Vertices", "Vertices" ) )
 				]
 				+SHorizontalBox::Slot()
 				.FillWidth( 1 )
 				.Padding( 2 )
 				[
-					SNew( SMeshEditorSelectionModeWidget, *MeshEditorMode, EEditableMeshElementType::Edge, LOCTEXT( "Edges", "Edges" ) )
+					SNew( SMeshEditorSelectionModeWidget, MeshEditorMode, EEditableMeshElementType::Edge, LOCTEXT( "Edges", "Edges" ) )
 				]
 				+SHorizontalBox::Slot()
 				.FillWidth( 1 )
 				.Padding( 2 )
 				[
-					SNew( SMeshEditorSelectionModeWidget, *MeshEditorMode, EEditableMeshElementType::Polygon, LOCTEXT( "Polygons", "Polygons" ) )
+					SNew( SMeshEditorSelectionModeWidget, MeshEditorMode, EEditableMeshElementType::Polygon, LOCTEXT( "Polygons", "Polygons" ) )
 				]
 				+SHorizontalBox::Slot()
 				.FillWidth( 1 )
 				.Padding( 2 )
 				[
-					SNew( SMeshEditorSelectionModeWidget, *MeshEditorMode, EEditableMeshElementType::Any, LOCTEXT( "Any", "Any" ) )
+					SNew( SMeshEditorSelectionModeWidget, MeshEditorMode, EEditableMeshElementType::Any, LOCTEXT( "Any", "Any" ) )
 				]
 			]
 			+SVerticalBox::Slot()
@@ -254,8 +251,8 @@ void SMeshEditorModeControls::Construct( const FArguments& InArgs )
 					.VAlign( VAlign_Center )
 					.Text( LOCTEXT( "Propagate", "Propagate" ) )
 					.ToolTip( SNew( SToolTip ).Text( LOCTEXT( "PropagateTooltip", "Propagates per-instance changes to the static mesh asset itself." ) ) )
-					.IsEnabled_Lambda( [ MeshEditorMode ] { return false; /*MeshEditorMode->CanPropagateInstanceChanges(); */ } )
-					.OnClicked_Lambda( [ MeshEditorMode ] { MeshEditorMode->PropagateInstanceChanges(); return FReply::Handled(); } )
+					.IsEnabled_Lambda( [ &MeshEditorMode ] { return false; /*MeshEditorMode.CanPropagateInstanceChanges(); */ } )
+					.OnClicked_Lambda( [ &MeshEditorMode ] { MeshEditorMode.PropagateInstanceChanges(); return FReply::Handled(); } )
 				]
 				+SHorizontalBox::Slot()
 				.AutoWidth()
@@ -269,8 +266,8 @@ void SMeshEditorModeControls::Construct( const FArguments& InArgs )
 						.Style( FMeshEditorStyle::Get(), "SelectionMode.Entry" )
 						.HAlign( HAlign_Fill )
 						.ToolTip( SNew( SToolTip ).Text( LOCTEXT( "PerInstanceTooltip", "Toggles editing mode between editing instances and editing the original static mesh asset." ) ) )
-						.IsChecked_Lambda( [ MeshEditorMode ] { return ( MeshEditorMode->IsEditingPerInstance() ? ECheckBoxState::Checked : ECheckBoxState::Unchecked ); } )
-						.OnCheckStateChanged_Lambda( [ MeshEditorMode ]( ECheckBoxState State ) { MeshEditorMode->SetEditingPerInstance( State == ECheckBoxState::Checked ); } )
+						.IsChecked_Lambda( [ &MeshEditorMode ] { return ( MeshEditorMode.IsEditingPerInstance() ? ECheckBoxState::Checked : ECheckBoxState::Unchecked ); } )
+						.OnCheckStateChanged_Lambda( [ &MeshEditorMode ]( ECheckBoxState State ) { MeshEditorMode.SetEditingPerInstance( State == ECheckBoxState::Checked ); } )
 						[
 							SNew( SHorizontalBox )
 							+SHorizontalBox::Slot()
@@ -310,7 +307,7 @@ void FMeshEditorModeToolkit::UnregisterTabSpawners( const TSharedRef<FTabManager
 
 void FMeshEditorModeToolkit::Init( const TSharedPtr<IToolkitHost>& InitToolkitHost )
 {
-	ToolkitWidget = SNew( SMeshEditorModeControls );
+	ToolkitWidget = SNew( SMeshEditorModeControls, MeshEditorMode );
 
 	FModeToolkit::Init( InitToolkitHost );
 }
diff --git a/Engine/Plugins/Editor/MeshEditor/Source/MeshEditor/MeshEditorModeToolkit.h b/Engine/Plugins/Editor/MeshEditor/Source/MeshEditor/MeshEditorModeToolkit.h
--- a/Engine/Plugins/Editor/MeshEditor/Source/MeshEditor/MeshEditorModeToolkit.h
+++ b/Engine/Plugins/Editor/MeshEditor/Source/MeshEditor/MeshEditorModeToolkit.h
@@ -2,9 +2,9 @@
 
 #pragma once
 #include "Toolkits/BaseToolkit.h"
-#include "MeshEditorMode.h"
 #include "EditorModeManager.h"
 
+
 /** Mesh Editor Mode widget for controls */
 class SMeshEditorModeControls : public SCompoundWidget
 {
@@ -15,16 +15,10 @@ public:
 public:
 
 	/** SCompoundWidget functions */
-	void Construct( const FArguments& InArgs );
+	void Construct( const FArguments& InArgs, class IMeshEditorModeUIContract& MeshEditorMode );
 
 protected:
 
-	FMeshEditorMode* GetMeshEditorMode() const
-	{
-		return static_cast<FMeshEditorMode*>( GLevelEditorModeTools().GetActiveMode( "MeshEditor" ) );
-	}
-
-private:
 };
 
 
@@ -35,7 +29,7 @@ class FMeshEditorModeToolkit : public FModeToolkit
 {
 public:
 
-	FMeshEditorModeToolkit( FMeshEditorMode& InMeshEditorMode )
+	FMeshEditorModeToolkit( class IMeshEditorModeUIContract& InMeshEditorMode )
 		: MeshEditorMode( InMeshEditorMode )
 	{}
 
@@ -57,5 +51,5 @@ private:
 	TSharedPtr<SMeshEditorModeControls> ToolkitWidget;
 
 	/** Reference to the owning FMeshEditorMode */
-	FMeshEditorMode& MeshEditorMode;
+	IMeshEditorModeUIContract& MeshEditorMode;
 };
diff --git a/Engine/Plugins/Editor/MeshEditor/Source/MeshEditor/MeshEditorModule.cpp b/Engine/Plugins/Editor/MeshEditor/Source/MeshEditor/MeshEditorModule.cpp
--- a/Engine/Plugins/Editor/MeshEditor/Source/MeshEditor/MeshEditorModule.cpp
+++ b/Engine/Plugins/Editor/MeshEditor/Source/MeshEditor/MeshEditorModule.cpp
@@ -13,6 +13,8 @@
 #include "IVREditorModule.h"
 #include "VREditorMode.h"
 #include "Editor.h"
+#include "IMeshEditorModeUIContract.h"
+
 
 #define LOCTEXT_NAMESPACE "MeshEditor"
 
@@ -191,9 +193,10 @@ void FMeshEditorModule::OnMeshEditModeButtonClicked(EEditableMeshElementType InM
 	}
 
 	FMeshEditorMode* MeshEditorMode = static_cast<FMeshEditorMode*>( GLevelEditorModeTools().FindMode( GetEditorModeID() ) );
-	if (MeshEditorMode != nullptr)
+	if ( MeshEditorMode != nullptr)
 	{
-		MeshEditorMode->SetMeshElementSelectionMode(InMode);
+		IMeshEditorModeUIContract* MeshEditorModeUIContract = (IMeshEditorModeUIContract*)MeshEditorMode;
+		MeshEditorModeUIContract->SetMeshElementSelectionMode(InMode);
 
 		UVREditorMode* VREditorMode = IVREditorModule::Get().GetVRMode();
 		if( VREditorMode != nullptr )
@@ -207,10 +210,11 @@ ECheckBoxState FMeshEditorModule::IsMeshEditModeButtonChecked(EEditableMeshEleme
 {
 	bool bMeshModeActive = false;
 
-	FMeshEditorMode* MeshEditorMode = static_cast<FMeshEditorMode*>( GLevelEditorModeTools().FindMode( GetEditorModeID() ) );
+	const FMeshEditorMode* MeshEditorMode = static_cast<FMeshEditorMode*>( GLevelEditorModeTools().FindMode( GetEditorModeID() ) );
 	if( MeshEditorMode != nullptr )
 	{
-		bMeshModeActive = MeshEditorMode->GetMeshElementSelectionMode() == InMode;
+		const IMeshEditorModeUIContract* MeshEditorModeUIContract = (const IMeshEditorModeUIContract*)MeshEditorMode;
+		bMeshModeActive = MeshEditorModeUIContract->GetMeshElementSelectionMode() == InMode;
 	}
 	return bMeshModeActive ? ECheckBoxState::Checked : ECheckBoxState::Unchecked;
 }
diff --git a/Engine/Plugins/Editor/MeshEditor/Source/MeshEditor/MeshElement.h b/Engine/Plugins/Editor/MeshEditor/Source/MeshEditor/MeshElement.h
deleted file mode 100644
diff --git a/Engine/Plugins/Editor/MeshEditor/Source/MeshEditor/Public/IMeshEditorModeEditingContract.h b/Engine/Plugins/Editor/MeshEditor/Source/MeshEditor/Public/IMeshEditorModeEditingContract.h
new file mode 100644
--- /dev/null
+++ b/Engine/Plugins/Editor/MeshEditor/Source/MeshEditor/Public/IMeshEditorModeEditingContract.h
@@ -0,0 +1,87 @@
+// Copyright 1998-2015 Epic Games, Inc. All Rights Reserved.
+
+#pragma once
+
+#include "EditableMeshTypes.h"
+#include "MeshElement.h"
+
+
+class IMeshEditorModeEditingContract
+{
+
+public:
+
+	enum class EMeshEditAction
+	{
+		/** Nothing going on right now */
+		None,
+
+		/** Selecting mesh elements by 'painting' over multiple elements */
+		SelectByPainting,
+
+		/** Moving elements using a transform gizmo */
+		MoveUsingGizmo,
+
+		/** Moving selected mesh elements (vertices, edges or polygons) */
+		Move,
+
+		/** Split an edge by inserting a vertex.  You can drag to preview where the vertex will be inserted. */
+		SplitEdge,
+
+		/** Splits an edge by inserting a new vertex, then immediately starts dragging that vertex */
+		SplitEdgeAndDragVertex,
+
+		/** Insert an edge loop */
+		InsertEdgeLoop,
+
+		/** Extrude polygon by making a copy of it and allowing you to shift it along the polygon normal axis */
+		ExtrudePolygon,
+
+		/** Extrude polygon by making a copy of it and allowing you to move it around freely */
+		FreelyExtrudePolygon,
+
+		/** Inset polygon by replacing it with a new polygon that is bordered by polygons of a specific relative size */
+		InsetPolygon,
+
+		/** Bevel polygons by adding angled bordering polygons of a specific relative size */
+		BevelPolygon,
+
+		/** Extend an edge by making a copy of it and allowing you to move it around */
+		ExtendEdge,
+
+		/** Extend a vertex by making a copy of it, creating new polygons to join the geometry together */
+		ExtendVertex,
+
+		/** For subdivision meshes, edits how sharp a vertex corner is by dragging in space */
+		EditVertexCornerSharpness,
+
+		/** For subdivision meshes, edits how sharp an edge crease is by dragging in space */
+		EditEdgeCreaseSharpness,
+
+		/** Freehand vertex drawing */
+		DrawVertices,
+	};
+
+
+	/** Gets the interactive action currently being performed (and previewed).  These usually happen over multiple frames, and
+	    result in a 'final' application of the change that performs a more exhaustive (and more expensive) update. */
+	virtual EMeshEditAction GetActiveAction() const = 0;
+
+	virtual void GetSelectedMeshesAndVertices( TMap<class UEditableMesh*, TArray<FMeshElement>>& OutMeshesAndVertices ) = 0;
+	virtual void GetSelectedMeshesAndEdges( TMap<class UEditableMesh*, TArray<FMeshElement>>& OutMeshesAndEdges ) = 0;
+	virtual void GetSelectedMeshesAndPolygons( TMap<class UEditableMesh*, TArray<FMeshElement>>& OutMeshesAndPolygons ) = 0;
+
+	/** Selects the specified mesh elements */
+	virtual void SelectMeshElements( const TArray<FMeshElement>& MeshElementsToSelect ) = 0;
+
+	/** Deselects the specified mesh elements */
+	virtual void DeselectMeshElements( const TArray<FMeshElement>& MeshElementsToDeselect ) = 0;
+	virtual void DeselectMeshElements( const TMap<UEditableMesh*, TArray<FMeshElement>>& MeshElementsToDeselect ) = 0;
+
+	/** Commits all selected meshes */
+	virtual void CommitSelectedMeshes() = 0;
+
+	/** Clears hover and selection on mesh elements that may no longer be valid.  You'll want to call this if you change the mesh topology */
+	virtual void ClearInvalidSelectedElements() = 0;
+
+};
diff --git a/Engine/Plugins/Editor/MeshEditor/Source/MeshEditor/Public/MeshEditorCommands.h b/Engine/Plugins/Editor/MeshEditor/Source/MeshEditor/Public/MeshEditorCommands.h
new file mode 100644
--- /dev/null
+++ b/Engine/Plugins/Editor/MeshEditor/Source/MeshEditor/Public/MeshEditorCommands.h
@@ -0,0 +1,201 @@
+// Copyright 1998-2015 Epic Games, Inc. All Rights Reserved.
+
+#pragma once
+
+#include "MeshEditorCommands.generated.h"
+
+
+UCLASS( abstract )
+class MESHEDITOR_API UMeshEditorCommand : public UObject
+{
+	GENERATED_BODY()
+
+public:
+
+	/** Registers the UI command for this mesh editor command */
+	virtual void RegisterUICommand( class FBindingContext* BindingContext ) PURE_VIRTUAL(,);
+
+	/** Runs this command */
+	virtual void Execute( class IMeshEditorModeEditingContract& MeshEditorMode, bool& bOutWasSuccessful ) PURE_VIRTUAL(,);
+
+	/** Gets the UI command info for this command */
+	const TSharedPtr<class FUICommandInfo>& GetUICommandInfo() const
+	{
+		return UICommandInfo;
+	}
+
+
+protected:
+
+	/** Our UI command for this action */
+	TSharedPtr<FUICommandInfo> UICommandInfo;
+};
+
+
+UCLASS( abstract )
+class MESHEDITOR_API UMeshEditorPolygonCommand : public UMeshEditorCommand
+{
+	GENERATED_BODY()
+
+public:
+
+};
+
+
+
+// Actions that can be invoked from this mode regardless of what type of elements are selected
+class FMeshEditorCommonCommands : public TCommands<FMeshEditorCommonCommands>
+{
+public:
+	FMeshEditorCommonCommands();
+
+	// TCommands<> interface
+	virtual void RegisterCommands() override;
+	// End of TCommands<> interface
+
+public:
+
+	/** Deletes selected mesh elements, including polygons partly defined by selected elements */
+	TSharedPtr<FUICommandInfo> DeleteMeshElement;
+
+	/** Increases the number of subdivision levels for the selected mesh. */
+	TSharedPtr<FUICommandInfo> AddSubdivisionLevel;
+
+	/** Decreases the number of subdivision levels for the selected mesh. */
+	TSharedPtr<FUICommandInfo> RemoveSubdivisionLevel;
+
+	/** Shows vertex normals */
+	TSharedPtr<FUICommandInfo> ShowVertexNormals;
+
+	/** Marquee select actions */
+	TSharedPtr<FUICommandInfo> MarqueeSelectVertices;
+	TSharedPtr<FUICommandInfo> MarqueeSelectEdges;
+	TSharedPtr<FUICommandInfo> MarqueeSelectPolygons;
+
+	/** Draw vertices */
+	TSharedPtr<FUICommandInfo> DrawVertices;
+
+	/** Frame selected elements */
+	TSharedPtr<FUICommandInfo> FrameSelectedElements;
+
+	/** Set mesh element selection modes */
+	TSharedPtr<FUICommandInfo> SetVertexSelectionMode;
+	TSharedPtr<FUICommandInfo> SetEdgeSelectionMode;
+	TSharedPtr<FUICommandInfo> SetPolygonSelectionMode;
+	TSharedPtr<FUICommandInfo> SetAnySelectionMode;
+
+	/** Quadrangulate mesh */
+	TSharedPtr<FUICommandInfo> QuadrangulateMesh;
+};
+
+
+// Actions that can be invoked from this mode when vertices are selected
+class FMeshEditorVertexCommands : public TCommands<FMeshEditorVertexCommands>
+{
+public:
+	FMeshEditorVertexCommands();
+
+	// TCommands<> interface
+	virtual void RegisterCommands() override;
+	// End of TCommands<> interface
+
+public:
+
+	/** Sets the primary action to move vertices */
+	TSharedPtr<FUICommandInfo> MoveVertex;
+
+	/** Sets the primary action to extend vertices */
+	TSharedPtr<FUICommandInfo> ExtendVertex;
+
+	/** Sets the primary action to edit the vertex's corner sharpness */
+	TSharedPtr<FUICommandInfo> EditVertexCornerSharpness;
+
+	/** Removes the selected vertex if possible */
+	TSharedPtr<FUICommandInfo> RemoveVertex;
+
+	/** Welds the selected vertices */
+	TSharedPtr<FUICommandInfo> WeldVertices;
+};
+
+
+// Actions that can be invoked from this mode when edges are selected
+class FMeshEditorEdgeCommands : public TCommands<FMeshEditorEdgeCommands>
+{
+public:
+	FMeshEditorEdgeCommands();
+
+	// TCommands<> interface
+	virtual void RegisterCommands() override;
+	// End of TCommands<> interface
+
+public:
+
+	/** Sets the primary action to move edges */
+	TSharedPtr<FUICommandInfo> MoveEdge;
+
+	/** Sets the primary action to split edges */
+	TSharedPtr<FUICommandInfo> SplitEdge;
+
+	/** Sets the primary action to split edges and drag vertices */
+	TSharedPtr<FUICommandInfo> SplitEdgeAndDragVertex;
+
+	/** Sets the primary action to insert edge loops */
+	TSharedPtr<FUICommandInfo> InsertEdgeLoop;
+
+	/** Sets the primary action to extend edges */
+	TSharedPtr<FUICommandInfo> ExtendEdge;
+
+	/** Sets the primary action to edit the edge's crease sharpness */
+	TSharedPtr<FUICommandInfo> EditEdgeCreaseSharpness;
+
+	/** Removes the selected edge if possible */
+	TSharedPtr<FUICommandInfo> RemoveEdge;
+
+	/** Soften edge */
+	TSharedPtr<FUICommandInfo> SoftenEdge;
+
+	/** Harden edge */
+	TSharedPtr<FUICommandInfo> HardenEdge;
+
+	/** Select edge loop */
+	TSharedPtr<FUICommandInfo> SelectEdgeLoop;
+};
+
+
+// Actions that can be invoked from this mode when polygons are selected
+class FMeshEditorPolygonCommands : public TCommands<FMeshEditorPolygonCommands>
+{
+public:
+	FMeshEditorPolygonCommands();
+
+	// TCommands<> interface
+	virtual void RegisterCommands() override;
+	// End of TCommands<> interface
+
+public:
+
+	/** Sets the primary action to move polygons */
+	TSharedPtr<FUICommandInfo> MovePolygon;
+
+	/** Sets the primary action to extrude polygons */
+	TSharedPtr<FUICommandInfo> ExtrudePolygon;
+
+	/** Sets the primary action to freely extrude polygons */
+	TSharedPtr<FUICommandInfo> FreelyExtrudePolygon;
+
+	/** Sets the primary action to inset polygons */
+	TSharedPtr<FUICommandInfo> InsetPolygon;
+
+	/** Sets the primary action to bevel polygons */
+	TSharedPtr<FUICommandInfo> BevelPolygon;
+
+	/** Flips the currently selected polygon(s) */
+	TSharedPtr<FUICommandInfo> FlipPolygon;
+
+	/** Triangulates the currently selected polygon(s) */
+	TSharedPtr<FUICommandInfo> TriangulatePolygon;
+
+	/** Assigns the highlighted material to the currently selected polygon(s) */
+	TSharedPtr<FUICommandInfo> AssignMaterial;
+};
+
diff --git a/Engine/Plugins/Editor/MeshEditor/Source/MeshEditor/Public/MeshElement.h b/Engine/Plugins/Editor/MeshEditor/Source/MeshEditor/Public/MeshElement.h
new file mode 100644
--- /dev/null
+++ b/Engine/Plugins/Editor/MeshEditor/Source/MeshEditor/Public/MeshElement.h
@@ -0,0 +1,178 @@
+// Copyright 1998-2015 Epic Games, Inc. All Rights Reserved.
+
+#pragma once
+
+#include "EditableMeshTypes.h"
+
+/**
+ * Uniquely identifies a specific element within a mesh
+ */
+struct FEditableMeshElementAddress
+{
+	/** The sub-mesh address that the element is contained by */
+	FEditableMeshSubMeshAddress SubMeshAddress;
+
+	/** The type of element */
+	EEditableMeshElementType ElementType;
+
+	/** For polygons, the section ID the polygon belongs to, otherwise FEditableMeshSectionID::Invalid */
+	FSectionID SectionID;
+
+	/** The ID of the element within the mesh */
+	FElementID ElementID;
+
+
+	/** Default constructor that initializes variables to an invalid element address */
+	FEditableMeshElementAddress()
+		: SubMeshAddress(),
+		  ElementType( EEditableMeshElementType::Invalid ),
+		  SectionID( FSectionID::Invalid ),
+		  ElementID( FElementID::Invalid )
+	{
+	}
+
+	FEditableMeshElementAddress( const FEditableMeshSubMeshAddress& InSubMeshAddress, FVertexID InVertexID )
+		: SubMeshAddress( InSubMeshAddress ),
+		  ElementType( EEditableMeshElementType::Vertex ),
+		  SectionID( FSectionID::Invalid ),
+		  ElementID( InVertexID )
+	{
+	}
+
+	FEditableMeshElementAddress( const FEditableMeshSubMeshAddress& InSubMeshAddress, FEdgeID InEdgeID )
+		: SubMeshAddress( InSubMeshAddress ),
+		  ElementType( EEditableMeshElementType::Edge ),
+		  SectionID( FSectionID::Invalid ),
+		  ElementID( InEdgeID )
+	{
+	}
+
+	FEditableMeshElementAddress( const FEditableMeshSubMeshAddress& InSubMeshAddress, FPolygonRef InPolygonRef )
+		: SubMeshAddress( InSubMeshAddress ),
+		  ElementType( EEditableMeshElementType::Polygon ),
+		  SectionID( InPolygonRef.SectionID ),
+		  ElementID( InPolygonRef.PolygonID )
+	{
+	}
+
+	/** Equality check */
+	inline bool operator==( const FEditableMeshElementAddress& Other ) const
+	{
+		return
+			SubMeshAddress == Other.SubMeshAddress &&
+			ElementType == Other.ElementType &&
+			SectionID == Other.SectionID &&
+			ElementID == Other.ElementID;
+	}
+
+	/** Convert to a string */
+	inline FString ToString() const
+	{
+		FString ElementTypeString;
+		switch( ElementType )
+		{
+			case EEditableMeshElementType::Invalid:
+				ElementTypeString = TEXT( "Invalid" );
+				break;
+
+			case EEditableMeshElementType::Vertex:
+				ElementTypeString = TEXT( "Vertex" );
+				break;
+
+			case EEditableMeshElementType::Edge:
+				ElementTypeString = TEXT( "Edge" );
+				break;
+
+			case EEditableMeshElementType::Polygon:
+				ElementTypeString = TEXT( "Polygon" );
+				break;
+
+			default:
+				check( 0 );	// Unrecognized type
+		}
+
+		return FString::Printf(
+			TEXT( "%s, ElementType:%s, SectionID:%s, ElementID:%s" ),
+			*SubMeshAddress.ToString(),
+			*ElementTypeString,
+			*SectionID.ToString(),
+			*ElementID.ToString() );
+	}
+};
+
+
+struct FMeshElement
+{
+	/** The component that is referencing the mesh.  Does not necessarily own the mesh!  The mesh could be shared
+		between many components. */
+	TWeakObjectPtr<class UPrimitiveComponent> Component;
+
+	/** The address of the mesh element */
+	FEditableMeshElementAddress ElementAddress;
+
+	/** Real time in seconds that we were last hovered over */
+	double LastHoverTime;
+
+	/** Real time in seconds that we were last selected */
+	double LastSelectTime;
+
+
+	/** Default constructor that initializes everything to safe values */
+	FMeshElement()
+		: Component( nullptr ),
+		  ElementAddress(),
+		  LastHoverTime( 0.0 ),
+		  LastSelectTime( 0.0 )
+	{
+	}
+
+	FMeshElement( UPrimitiveComponent* InComponent, const FEditableMeshSubMeshAddress& InSubMeshAddress, FVertexID InVertexID, double InLastHoverTime = 0.0, double InLastSelectTime = 0.0 )
+		: Component( InComponent ),
+		  ElementAddress( InSubMeshAddress, InVertexID ),
+		  LastHoverTime( InLastHoverTime ),
+		  LastSelectTime( InLastSelectTime )
+	{
+	}
+
+	FMeshElement( UPrimitiveComponent* InComponent, const FEditableMeshSubMeshAddress& InSubMeshAddress, FEdgeID InEdgeID, double InLastHoverTime = 0.0, double InLastSelectTime = 0.0 )
+		: Component( InComponent ),
+		  ElementAddress( InSubMeshAddress, InEdgeID ),
+		  LastHoverTime( InLastHoverTime ),
+		  LastSelectTime( InLastSelectTime )
+	{
+	}
+
+	FMeshElement( UPrimitiveComponent* InComponent, const FEditableMeshSubMeshAddress& InSubMeshAddress, FPolygonRef InPolygonRef, double InLastHoverTime = 0.0, double InLastSelectTime = 0.0 )
+		: Component( InComponent ),
+		  ElementAddress( InSubMeshAddress, InPolygonRef ),
+		  LastHoverTime( InLastHoverTime ),
+		  LastSelectTime( InLastSelectTime )
+	{
+	}
+
+	/** Checks to see if we have something valid */
+	inline bool IsValidMeshElement() const
+	{
+		return
+			( Component.IsValid() &&
+				ElementAddress.SubMeshAddress.EditableMeshFormat != nullptr &&
+				ElementAddress.ElementType != EEditableMeshElementType::Invalid );
+	}
+
+	/** Checks to see if this mesh element points to the same element as another mesh element */
+	inline bool IsSameMeshElement( const FMeshElement& Other ) const
+	{
+		// NOTE: We only care that the element addresses are the same, not other transient state
+		return Component == Other.Component && ElementAddress == Other.ElementAddress;
+	}
+
+	/** Convert to a string */
+	FString ToString() const
+	{
+		return FString::Printf(
+			TEXT( "Component:%s, %s" ),
+			*Component->GetName(),
+			*ElementAddress.ToString() );
+	}
+};
+
diff --git a/Engine/Plugins/Editor/MeshEditor/Source/MeshEditor/Public/UndoHelpers.h b/Engine/Plugins/Editor/MeshEditor/Source/MeshEditor/Public/UndoHelpers.h
new file mode 100644
--- /dev/null
+++ b/Engine/Plugins/Editor/MeshEditor/Source/MeshEditor/Public/UndoHelpers.h
@@ -0,0 +1,15 @@
+// Copyright 1998-2015 Epic Games, Inc. All Rights Reserved.
+
+#pragma once
+
+#include "CoreMinimal.h"
+
+
+namespace UndoHelpers
+{
+	/** @return Returns true if the undo system is available right now.  When in Simulate Mode, we can't store undo states or use undo/redo features */
+	extern MESHEDITOR_API bool IsUndoSystemAvailable();
+
+	/** Saves undo state, if possible (e.g., not in Simulate mode.) */
+	extern MESHEDITOR_API void StoreUndo( UObject* Object, TUniquePtr<FChange> UndoChange );
+}
\ No newline at end of file
diff --git a/Engine/Plugins/Editor/MeshEditor/Source/MeshEditor/UndoHelpers.cpp b/Engine/Plugins/Editor/MeshEditor/Source/MeshEditor/UndoHelpers.cpp
new file mode 100644
--- /dev/null
+++ b/Engine/Plugins/Editor/MeshEditor/Source/MeshEditor/UndoHelpers.cpp
@@ -0,0 +1,26 @@
+// Copyright 1998-2016 Epic Games, Inc. All Rights Reserved.
+
+#include "UndoHelpers.h"
+
+namespace UndoHelpers
+{
+	bool IsUndoSystemAvailable()
+	{
+		return GUndo != nullptr;
+	}
+
+
+	void StoreUndo( UObject* Object, TUniquePtr<FChange> UndoChange )
+	{
+		// Did you forget to use an FScopedTransaction?  If GUndo was null, then most likely we forgot to wrap this call within an editor transaction.
+		// The only exception is in Simulate mode, where Undo is not allowed.
+		check( IsUndoSystemAvailable() || GEditor == nullptr || GEditor->bIsSimulatingInEditor );
+
+		if( IsUndoSystemAvailable() )
+		{
+			GUndo->StoreUndo( Object, MoveTemp( UndoChange ) );
+		}
+	}
+}
+
+
diff --git a/Engine/Plugins/Editor/MeshEditor/Source/PolygonModeling/PolygonModeling.Build.cs b/Engine/Plugins/Editor/MeshEditor/Source/PolygonModeling/PolygonModeling.Build.cs
new file mode 100644
--- /dev/null
+++ b/Engine/Plugins/Editor/MeshEditor/Source/PolygonModeling/PolygonModeling.Build.cs
@@ -0,0 +1,22 @@
+// Copyright 1998-2015 Epic Games, Inc. All Rights Reserved.
+
+namespace UnrealBuildTool.Rules
+{
+	public class PolygonModeling : ModuleRules
+	{
+        public PolygonModeling(ReadOnlyTargetRules Target) : base(Target)
+		{
+			PrivateDependencyModuleNames.AddRange(
+                new string[] {
+				    "Core",
+				    "CoreUObject",
+					"Slate",
+					"Engine",
+				    "UnrealEd",
+					"MeshEditingRuntime",
+					"MeshEditor"
+				}
+			);
+		}
+	}
+}
\ No newline at end of file
diff --git a/Engine/Plugins/Editor/MeshEditor/Source/PolygonModeling/PolygonModelingModule.cpp b/Engine/Plugins/Editor/MeshEditor/Source/PolygonModeling/PolygonModelingModule.cpp
new file mode 100644
--- /dev/null
+++ b/Engine/Plugins/Editor/MeshEditor/Source/PolygonModeling/PolygonModelingModule.cpp
@@ -0,0 +1,8 @@
+// Copyright 1998-2016 Epic Games, Inc. All Rights Reserved.
+
+#include "CoreMinimal.h"
+#include "ModuleManager.h"
+
+
+IMPLEMENT_MODULE( FDefaultModuleImpl, PolygonModeling )
+
diff --git a/Engine/Plugins/Editor/MeshEditor/Source/PolygonModeling/TessellatePolygons.cpp b/Engine/Plugins/Editor/MeshEditor/Source/PolygonModeling/TessellatePolygons.cpp
new file mode 100644
--- /dev/null
+++ b/Engine/Plugins/Editor/MeshEditor/Source/PolygonModeling/TessellatePolygons.cpp
@@ -0,0 +1,102 @@
+// Copyright 1998-2016 Epic Games, Inc. All Rights Reserved.
+
+#include "TessellatePolygons.h"
+#include "IMeshEditorModeEditingContract.h"
+#include "UndoHelpers.h"
+#include "UICommandInfo.h"
+#include "EditableMesh.h"
+#include "MeshElement.h"
+
+
+#define LOCTEXT_NAMESPACE "MeshEditorMode"
+
+
+void UTessellatePolygonsCommand::RegisterUICommand( FBindingContext* BindingContext )
+{
+	UI_COMMAND_EXT( BindingContext, UICommandInfo, "Tessellate Selected Polygons", "Tessellate selected polygons into smaller polygons.", EUserInterfaceActionType::Button, FInputChord() );
+}
+
+
+void UTessellatePolygonsCommand::Execute( IMeshEditorModeEditingContract& MeshEditorMode, bool& bOutWasSuccessful )
+{
+	bOutWasSuccessful = false;
+	if( MeshEditorMode.GetActiveAction() != IMeshEditorModeEditingContract::EMeshEditAction::None )
+	{
+		return;
+	}
+
+	static TMap<UEditableMesh*, TArray<FMeshElement>> SelectedMeshesAndPolygons;
+	MeshEditorMode.GetSelectedMeshesAndPolygons( SelectedMeshesAndPolygons );
+	if( SelectedMeshesAndPolygons.Num() == 0 )
+	{
+		return;
+	}
+
+	FScopedTransaction Transaction( LOCTEXT( "UndoTessellatePolygon", "Tessellate Polygon" ) );
+
+	MeshEditorMode.CommitSelectedMeshes();
+
+	// Refresh selection (committing may have created a new mesh instance)
+	MeshEditorMode.GetSelectedMeshesAndPolygons( SelectedMeshesAndPolygons );
+
+	// Deselect the mesh elements before we delete them.  This will make sure they become selected again after undo.
+	MeshEditorMode.DeselectMeshElements( SelectedMeshesAndPolygons );
+
+	static TArray<FMeshElement> MeshElementsToSelect;
+	MeshElementsToSelect.Reset();
+
+	for( const auto& SelectedMeshAndPolygons : SelectedMeshesAndPolygons )
+	{
+		UEditableMesh* EditableMesh = SelectedMeshAndPolygons.Key;
+		const TArray<FMeshElement>& PolygonElements = SelectedMeshAndPolygons.Value;
+
+		static TArray<FPolygonRef> PolygonsToTessellate;
+		PolygonsToTessellate.Reset( PolygonElements.Num() );
+
+		for( const auto& PolygonElement : PolygonElements )
+		{
+			const FPolygonRef PolygonRef( PolygonElement.ElementAddress.SectionID, FPolygonID( PolygonElement.ElementAddress.ElementID ) );
+			PolygonsToTessellate.Add( PolygonRef );
+		}
+
+
+		EditableMesh->StartModification( EMeshModificationType::Final, EMeshTopologyChange::TopologyChange );
+
+		// @todo mesheditor: Expose as configurable parameter
+		const ETriangleTessellationMode TriangleTessellationMode = ETriangleTessellationMode::FourTriangles;
+
+		static TArray<FPolygonRef> NewPolygonRefs;
+		EditableMesh->TessellatePolygons( PolygonsToTessellate, TriangleTessellationMode, /* Out */ NewPolygonRefs );
+
+		// Select the new polygons
+		for( const FPolygonRef NewPolygonRef : NewPolygonRefs )
+		{
+			FMeshElement NewPolygonMeshElement;
+			{
+				NewPolygonMeshElement.Component = PolygonElements[0].Component;
+				NewPolygonMeshElement.ElementAddress = PolygonElements[0].ElementAddress;
+				NewPolygonMeshElement.ElementAddress.ElementType = EEditableMeshElementType::Polygon;
+				NewPolygonMeshElement.ElementAddress.SectionID = NewPolygonRef.SectionID;
+				NewPolygonMeshElement.ElementAddress.ElementID = NewPolygonRef.PolygonID;
+			}
+
+			MeshElementsToSelect.Add( NewPolygonMeshElement );
+		}
+
+		EditableMesh->EndModification();
+
+		UndoHelpers::StoreUndo( EditableMesh, EditableMesh->MakeUndo() );
+	}
+
+	// Make sure we're not still hovering over a polygon we removed
+	MeshEditorMode.ClearInvalidSelectedElements();
+
+	// Select the new smaller polygons
+	MeshEditorMode.SelectMeshElements( MeshElementsToSelect );
+
+	bOutWasSuccessful = true;
+}
+
+
+#undef LOCTEXT_NAMESPACE
+
diff --git a/Engine/Plugins/Editor/MeshEditor/Source/PolygonModeling/TessellatePolygons.h b/Engine/Plugins/Editor/MeshEditor/Source/PolygonModeling/TessellatePolygons.h
new file mode 100644
--- /dev/null
+++ b/Engine/Plugins/Editor/MeshEditor/Source/PolygonModeling/TessellatePolygons.h
@@ -0,0 +1,21 @@
+// Copyright 1998-2016 Epic Games, Inc. All Rights Reserved.
+
+#pragma once
+
+#include "MeshEditorCommands.h"
+#include "TessellatePolygons.generated.h"
+
+
+/** Tessellates selected polygons into smaller polygons */
+UCLASS()
+class UTessellatePolygonsCommand : public UMeshEditorPolygonCommand
+{
+	GENERATED_BODY()
+
+protected:
+
+	// Overrides
+	virtual void RegisterUICommand( class FBindingContext* BindingContext ) override;
+	virtual void Execute( class IMeshEditorModeEditingContract& MeshEditorMode, bool& bOutWasSuccessful ) override;
+
+};
diff --git a/Engine/Source/Runtime/MeshEditingRuntime/Public/EditableMeshTypes.h b/Engine/Source/Runtime/MeshEditingRuntime/Public/EditableMeshTypes.h
--- a/Engine/Source/Runtime/MeshEditingRuntime/Public/EditableMeshTypes.h
+++ b/Engine/Source/Runtime/MeshEditingRuntime/Public/EditableMeshTypes.h
@@ -5,7 +5,6 @@
 #include "CoreMinimal.h"
 #include "UObject/Object.h"
 #include "UObject/ScriptMacros.h"
-#include "Materials/MaterialInterface.h"
 #include "EditableMeshTypes.generated.h"
 
 // @todo mesheditor: Move elsewhere
@@ -1149,7 +1148,7 @@ struct FSectionToCreate
 
 	/** Material to assign to the new section */
 	UPROPERTY( BlueprintReadWrite, Category="Editable Mesh" )
-	UMaterialInterface* Material;
+	class UMaterialInterface* Material;
 
 	/** Whether the new section should have collision enabled */
 	UPROPERTY( BlueprintReadWrite, Category="Editable Mesh" )
diff --git a/Engine/Source/Runtime/Slate/Public/Framework/Commands/Commands.h b/Engine/Source/Runtime/Slate/Public/Framework/Commands/Commands.h
--- a/Engine/Source/Runtime/Slate/Public/Framework/Commands/Commands.h
+++ b/Engine/Source/Runtime/Slate/Public/Framework/Commands/Commands.h
@@ -15,9 +15,11 @@
 SLATE_API void UI_COMMAND_Function( FBindingContext* This, TSharedPtr< FUICommandInfo >& OutCommand, const TCHAR* OutSubNamespace, const TCHAR* OutCommandName, const TCHAR* OutCommandNameUnderscoreTooltip, const ANSICHAR* DotOutCommandName, const TCHAR* FriendlyName, const TCHAR* InDescription, const EUserInterfaceActionType::Type CommandType, const FInputChord& InDefaultChord );
 
 // This macro requires LOCTEXT_NAMESPACE to be defined. If you don't want the command to be placed under a sub namespace, provide "" as the namespace.
+#define UI_COMMAND_EXT( BindingContext, CommandId, FriendlyName, InDescription, CommandType, InDefaultChord ) \
+	UI_COMMAND_Function( BindingContext, CommandId, TEXT(LOCTEXT_NAMESPACE), TEXT(#CommandId), TEXT(#CommandId) TEXT("_ToolTip"), "." #CommandId, TEXT(FriendlyName), TEXT(InDescription), CommandType, InDefaultChord );
 #define UI_COMMAND( CommandId, FriendlyName, InDescription, CommandType, InDefaultChord ) \
-	UI_COMMAND_Function( this, CommandId, TEXT(LOCTEXT_NAMESPACE), TEXT(#CommandId), TEXT(#CommandId) TEXT("_ToolTip"), "." #CommandId, TEXT(FriendlyName), TEXT(InDescription), CommandType, InDefaultChord );
-
+	UI_COMMAND_EXT( this, CommandId, FriendlyName, InDescription, CommandType, InDefaultChord )
+	
 #undef LOC_DEFINE_REGION
 
 /** A base class for a set of commands. Inherit from it to make a set of commands. See MainFrameActions for an example. */
