Mesh Editor extensibility (phase 2)
- Insert Edge Loop feature has been moved to the PolygonModeling module
- Mesh Editor now has support for editing 'modes' that are registered externally, including VR Editor support.  Lots of changes to UMeshEditorCommand and derived classes.
- Removed UndoHelpers; All Mesh Editor undo now goes through FMeshEditMode::TrackUndo().  This is so that we can use the same code path for regular undo as we use for changes that we'll revert during preview
- Cleaned up some code to use helper methods for selecting/deselecting mesh elements
- Removed all kinds of code that was using TTuples for no good reason
- Fixed a few bugs with the previous extensibility checkin

diff --git a/Engine/Plugins/Editor/MeshEditor/Source/MeshEditor/IMeshEditorModeUIContract.h b/Engine/Plugins/Editor/MeshEditor/Source/MeshEditor/IMeshEditorModeUIContract.h
deleted file mode 100644
diff --git a/Engine/Plugins/Editor/MeshEditor/Source/MeshEditor/MeshEditorCommands.cpp b/Engine/Plugins/Editor/MeshEditor/Source/MeshEditor/MeshEditorCommands.cpp
--- a/Engine/Plugins/Editor/MeshEditor/Source/MeshEditor/MeshEditorCommands.cpp
+++ b/Engine/Plugins/Editor/MeshEditor/Source/MeshEditor/MeshEditorCommands.cpp
@@ -2,9 +2,44 @@
 
 #include "MeshEditorCommands.h"
 #include "MeshEditorStyle.h"
+#include "IMeshEditorModeUIContract.h"
+
 
 #define LOCTEXT_NAMESPACE "MeshEditorCommands"
 
+
+FUIAction UMeshEditorCommand::MakeUIAction( IMeshEditorModeUIContract& MeshEditorMode )
+{
+	const EEditableMeshElementType ElementType = GetElementType();
+
+	FUIAction UIAction;
+	if( IsMode() )
+	{
+		const EMeshEditAction::Type MeshEditAction = GetCommandName();
+
+		UIAction = FUIAction(
+			FExecuteAction::CreateLambda( [&MeshEditorMode, ElementType, MeshEditAction] { MeshEditorMode.SetEquippedAction( ElementType, MeshEditAction ); } ),
+			FCanExecuteAction::CreateLambda( [&MeshEditorMode, ElementType] { return MeshEditorMode.IsMeshElementTypeSelectedOrIsActiveSelectionMode( ElementType ); } ),
+			FIsActionChecked::CreateLambda( [&MeshEditorMode, ElementType, MeshEditAction] { return ( MeshEditorMode.GetEquippedAction( ElementType ) == MeshEditAction ); } )
+		);
+	}
+	else
+	{
+		FExecuteAction ExecuteAction( FExecuteAction::CreateLambda( [&MeshEditorMode, this]
+		{
+			this->Execute( MeshEditorMode );
+		} ) );
+
+		UIAction = FUIAction(
+			ExecuteAction,
+			FCanExecuteAction::CreateLambda( [&MeshEditorMode, ElementType] { return MeshEditorMode.IsMeshElementTypeSelected( ElementType ); } )
+		);
+	}
+
+	return UIAction;
+}
+
+
 FMeshEditorCommonCommands::FMeshEditorCommonCommands() 
 	: TCommands<FMeshEditorCommonCommands>(
 		"MeshEditorCommon",
@@ -51,6 +86,16 @@ void FMeshEditorVertexCommands::RegisterCommands()
 
 	UI_COMMAND(RemoveVertex, "Remove Vertex", "Remove the selected vertex if possible.", EUserInterfaceActionType::Button, FInputChord(EKeys::BackSpace));
 	UI_COMMAND(WeldVertices, "Weld Vertices", "Weld the selected vertices, keeping the first selected vertex.", EUserInterfaceActionType::Button, FInputChord());
+
+	// @todo mesheditor extensibility: What's the plan for default keybinds for commands registered in a modular way?  Should we suggest available keys?
+	for( TObjectIterator<UMeshEditorVertexCommand> VertexCommandCDOIter( RF_NoFlags ); VertexCommandCDOIter; ++VertexCommandCDOIter )
+	{
+		UMeshEditorVertexCommand* VertexCommandCDO = *VertexCommandCDOIter;
+		if( !( VertexCommandCDO->GetClass()->GetClassFlags() & CLASS_Abstract ) )
+		{
+			VertexCommandCDO->RegisterUICommand( this );
+		}
+	}
 }
 
 FMeshEditorEdgeCommands::FMeshEditorEdgeCommands()
@@ -68,7 +113,6 @@ void FMeshEditorEdgeCommands::RegisterCommands()
 	UI_COMMAND(MoveEdge, "Move Edge Mode", "Set the primary action to move edges.", EUserInterfaceActionType::RadioButton, FInputChord(EKeys::F1));
 	UI_COMMAND(SplitEdge, "Split Edge Mode", "Set the primary action to split edges.", EUserInterfaceActionType::RadioButton, FInputChord(EKeys::F2));
 	UI_COMMAND(SplitEdgeAndDragVertex, "Split Edge and Drag Vertex Mode", "Set the primary action to split edges and drag vertices.", EUserInterfaceActionType::RadioButton, FInputChord(EKeys::F3));
-	UI_COMMAND(InsertEdgeLoop, "Insert Edge Loop Mode", "Set the primary action to insert edge loops.", EUserInterfaceActionType::RadioButton, FInputChord(EKeys::F4));
 	UI_COMMAND(ExtendEdge, "Extend Edge Mode", "Set the primary action to extend edges.", EUserInterfaceActionType::RadioButton, FInputChord(EKeys::F5));
 	UI_COMMAND(EditEdgeCreaseSharpness, "Edit Edge Crease Sharpness Mode", "Set the primary action to edit the edge's crease sharpness.", EUserInterfaceActionType::RadioButton, FInputChord(EKeys::F6));
 
@@ -76,6 +120,15 @@ void FMeshEditorEdgeCommands::RegisterCommands()
 	UI_COMMAND(SoftenEdge, "Soften Edge", "Make selected edge soft.", EUserInterfaceActionType::Button, FInputChord(EKeys::H, EModifierKey::Shift));
 	UI_COMMAND(HardenEdge, "Harden Edge", "Make selected edge hard.", EUserInterfaceActionType::Button, FInputChord(EKeys::H));
 	UI_COMMAND(SelectEdgeLoop, "Select Edge Loop", "Select the edge loops which contain the selected edges.", EUserInterfaceActionType::Button, FInputChord(EKeys::Two, EModifierKey::Shift));
+
+	for( TObjectIterator<UMeshEditorEdgeCommand> EdgeCommandCDOIter( RF_NoFlags ); EdgeCommandCDOIter; ++EdgeCommandCDOIter )
+	{
+		UMeshEditorEdgeCommand* EdgeCommandCDO = *EdgeCommandCDOIter;
+		if( !( EdgeCommandCDO->GetClass()->GetClassFlags() & CLASS_Abstract ) )
+		{
+			EdgeCommandCDO->RegisterUICommand( this );
+		}
+	}
 }
 
 FMeshEditorPolygonCommands::FMeshEditorPolygonCommands() 
diff --git a/Engine/Plugins/Editor/MeshEditor/Source/MeshEditor/MeshEditorMode.cpp b/Engine/Plugins/Editor/MeshEditor/Source/MeshEditor/MeshEditorMode.cpp
--- a/Engine/Plugins/Editor/MeshEditor/Source/MeshEditor/MeshEditorMode.cpp
+++ b/Engine/Plugins/Editor/MeshEditor/Source/MeshEditor/MeshEditorMode.cpp
@@ -35,11 +35,31 @@
 #include "VREditorWorldInteraction.h"
 #include "MeshEditorAssetContainer.h"
 #include "Kismet/GameplayStatics.h"
-#include "UndoHelpers.h"
 
 
 #define LOCTEXT_NAMESPACE "MeshEditorMode"
 
+
+// @todo mesheditor extensibility: This should probably be removed after we've evicted all current mesh editing actions to another module
+namespace EMeshEditAction
+{
+	const Type None( "None" );
+	const Type SelectByPainting( "SelectByPainting" );
+	const Type MoveUsingGizmo( "MoveUsingGizmo" );
+	const Type Move( "Move" );
+	const Type SplitEdge( "SplitEdge" );
+	const Type SplitEdgeAndDragVertex( "SplitEdgeAndDragVertex" );
+	const Type ExtrudePolygon( "ExtrudePolygon" );
+	const Type FreelyExtrudePolygon( "FreelyExtrudePolygon" );
+	const Type InsetPolygon( "InsetPolygon" );
+	const Type BevelPolygon( "BevelPolygon" );
+	const Type ExtendEdge( "ExtendEdge" );
+	const Type ExtendVertex( "ExtendVertex" );
+	const Type EditVertexCornerSharpness( "EditVertexCornerSharpness" );
+	const Type EditEdgeCreaseSharpness( "EditEdgeCreaseSharpness" );
+	const Type DrawVertices( "DrawVertices" );
+}
+
 // @todo mesheditor urgent: Texture coords are getting trashed all the time (only for meshes with 1 UV set though?)
 
 namespace MeshEd
@@ -369,6 +389,10 @@ FMeshEditorMode::FMeshEditorMode()
 	  EquippedEdgeAction( EMeshEditAction::Move ),
 	  EquippedPolygonAction( EMeshEditAction::Move ),
 	  ActiveAction( EMeshEditAction::None ),
+	  bIsCapturingUndoForPreview( false ),
+	  PreviewRevertChanges(),
+	  ActiveActionModifiedMeshes(),
+	  MeshEditorModeProxyObject( nullptr ),
 	  ActiveActionInteractor( nullptr ),
 	  bActiveActionNeedsHoverLocation( false ),
 	  bIsFirstActiveActionUpdate( false ),
@@ -378,9 +402,6 @@ FMeshEditorMode::FMeshEditorMode()
 	  SplitEdgeMeshesAndEdgesToSplit(),
 	  SplitEdgeSplitList(),
 	  EditSharpnessStartLocation( FVector::ZeroVector ),
-	  PreviewRevertChanges(),
-	  PreviewTopologyChange( EMeshTopologyChange::NoTopologyChange ),
-	  MeshEditorModeProxyObject( nullptr ),
 	  SelectingByPaintingRevertChangeInput( nullptr ),
 	  bShowVertexNormals( false ),
 	  bMarqueeSelectTransactionActive( false ),
@@ -468,6 +489,7 @@ void FMeshEditorMode::RemoveEditableMeshReferences()
 	SelectedEdges.Empty();
 	SelectedPolygons.Empty();
 	PreviewRevertChanges.Empty();
+	ActiveActionModifiedMeshes.Reset();
 
 	if( ViewportWorldInteraction != nullptr )
 	{
@@ -475,7 +497,7 @@ void FMeshEditorMode::RemoveEditableMeshReferences()
 	}
 }
 
-void FMeshEditorMode::PlayStartActionSound(EMeshEditAction NewAction, UViewportInteractor* ActionInteractor)
+void FMeshEditorMode::PlayStartActionSound(EMeshEditAction::Type NewAction, UViewportInteractor* ActionInteractor)
 {
 	if (ActionInteractor != nullptr)
 	{
@@ -489,7 +511,7 @@ void FMeshEditorMode::PlayStartActionSound(EMeshEditAction NewAction, UViewportI
 	}
 }
 
-void FMeshEditorMode::PlayFinishActionSound(EMeshEditAction NewAction, UViewportInteractor* ActionInteractor)
+void FMeshEditorMode::PlayFinishActionSound(EMeshEditAction::Type NewAction, UViewportInteractor* ActionInteractor)
 {
 	if (ActionInteractor != nullptr)
 	{
@@ -518,7 +540,6 @@ void FMeshEditorMode::BindCommands()
 	RegisterEdgeEditingMode( MeshEditorEdgeActions.MoveEdge, EMeshEditAction::Move );
 	RegisterEdgeEditingMode( MeshEditorEdgeActions.SplitEdge, EMeshEditAction::SplitEdge );
 	RegisterEdgeEditingMode( MeshEditorEdgeActions.SplitEdgeAndDragVertex, EMeshEditAction::SplitEdgeAndDragVertex );
-	RegisterEdgeEditingMode( MeshEditorEdgeActions.InsertEdgeLoop, EMeshEditAction::InsertEdgeLoop );
 	RegisterEdgeEditingMode( MeshEditorEdgeActions.ExtendEdge, EMeshEditAction::ExtendEdge );
 	RegisterEdgeEditingMode( MeshEditorEdgeActions.EditEdgeCreaseSharpness, EMeshEditAction::EditEdgeCreaseSharpness );
 
@@ -564,20 +585,33 @@ void FMeshEditorMode::BindCommands()
 	RegisterPolygonCommand( MeshEditorPolygonActions.TriangulatePolygon, FExecuteAction::CreateLambda( [this] { TriangulateSelectedPolygons(); } ) );
 	RegisterPolygonCommand( MeshEditorPolygonActions.AssignMaterial, FExecuteAction::CreateLambda( [this] { AssignSelectedMaterialToSelectedPolygons(); } ) );
 
-	for( TObjectIterator<UMeshEditorPolygonCommand> PolygonCommandCDOIter( RF_NoFlags ); PolygonCommandCDOIter; ++PolygonCommandCDOIter )
+	for( TObjectIterator<UMeshEditorCommand> CommandCDOIter( RF_NoFlags ); CommandCDOIter; ++CommandCDOIter )
 	{
-		UMeshEditorPolygonCommand* PolygonCommandCDO = *PolygonCommandCDOIter;
-		if( !( PolygonCommandCDO->GetClass()->GetClassFlags() & CLASS_Abstract ) )
+		UMeshEditorCommand* CommandCDO = *CommandCDOIter;
+		if( !( CommandCDO->GetClass()->GetClassFlags() & CLASS_Abstract ) )
 		{
-			RegisterPolygonCommand( PolygonCommandCDO->GetUICommandInfo(), FExecuteAction::CreateLambda( [this, PolygonCommandCDO]
+			const EEditableMeshElementType ElementType = CommandCDO->GetElementType();
+			if( ElementType == EEditableMeshElementType::Vertex )
 			{
-				bool bWasSuccessful = false;
-				PolygonCommandCDO->Execute( *this, /* Out */ bWasSuccessful );
-			} ) );
+				VertexActions.Emplace( CommandCDO->GetUICommandInfo(), CommandCDO->MakeUIAction( *this ) );
+			}
+			else if( ElementType == EEditableMeshElementType::Edge )
+			{
+				EdgeActions.Emplace( CommandCDO->GetUICommandInfo(), CommandCDO->MakeUIAction( *this ));
+			}
+			else if( ElementType == EEditableMeshElementType::Polygon )
+			{
+				PolygonActions.Emplace( CommandCDO->GetUICommandInfo(), CommandCDO->MakeUIAction( *this ) );
+			}
+			else
+			{
+				check( 0 );
+			}
 		}
 	}
 
 
+
 	// Bind common actions
 	CommonCommands = MakeShared<FUICommandList>();
 	for( const TTuple<TSharedPtr<FUICommandInfo>, FUIAction>& CommonAction : CommonActions )
@@ -608,37 +642,37 @@ void FMeshEditorMode::BindCommands()
 }
 
 
-void FMeshEditorMode::RegisterCommonEditingMode( const TSharedPtr<FUICommandInfo>& Command, EMeshEditAction EditingMode )
+void FMeshEditorMode::RegisterCommonEditingMode( const TSharedPtr<FUICommandInfo>& Command, EMeshEditAction::Type EditingMode )
 {
 	RegisterVertexEditingMode( Command, EditingMode );
 	RegisterEdgeEditingMode( Command, EditingMode );
 	RegisterPolygonEditingMode( Command, EditingMode );
 }
 
-void FMeshEditorMode::RegisterVertexEditingMode( const TSharedPtr<FUICommandInfo>& Command, EMeshEditAction EditingMode )
+void FMeshEditorMode::RegisterVertexEditingMode( const TSharedPtr<FUICommandInfo>& Command, EMeshEditAction::Type EditingMode )
 {
 	VertexActions.Emplace( Command, FUIAction(
-		FExecuteAction::CreateLambda( [this, EditingMode] { EquippedVertexAction = EditingMode; } ),
+		FExecuteAction::CreateLambda( [this, EditingMode] { SetEquippedAction( EEditableMeshElementType::Vertex, EditingMode ); } ),
 		FCanExecuteAction::CreateLambda( [this] { return IsMeshElementTypeSelectedOrIsActiveSelectionMode( EEditableMeshElementType::Vertex ); } ),
 		FIsActionChecked::CreateLambda( [this, EditingMode] { return ( EquippedVertexAction == EditingMode ); } )
 	) );
 }
 
 
-void FMeshEditorMode::RegisterEdgeEditingMode( const TSharedPtr<FUICommandInfo>& Command, EMeshEditAction EditingMode )
+void FMeshEditorMode::RegisterEdgeEditingMode( const TSharedPtr<FUICommandInfo>& Command, EMeshEditAction::Type EditingMode )
 {
 	EdgeActions.Emplace( Command, FUIAction(
-		FExecuteAction::CreateLambda( [this, EditingMode] { EquippedEdgeAction = EditingMode; } ),
+		FExecuteAction::CreateLambda( [this, EditingMode] { SetEquippedAction( EEditableMeshElementType::Edge, EditingMode ); } ),
 		FCanExecuteAction::CreateLambda( [this] { return IsMeshElementTypeSelectedOrIsActiveSelectionMode( EEditableMeshElementType::Edge ); } ),
 		FIsActionChecked::CreateLambda( [this, EditingMode] { return ( EquippedEdgeAction == EditingMode ); } )
 	) );
 }
 
 
-void FMeshEditorMode::RegisterPolygonEditingMode( const TSharedPtr<FUICommandInfo>& Command, EMeshEditAction EditingMode )
+void FMeshEditorMode::RegisterPolygonEditingMode( const TSharedPtr<FUICommandInfo>& Command, EMeshEditAction::Type EditingMode )
 {
 	PolygonActions.Emplace( Command, FUIAction(
-		FExecuteAction::CreateLambda( [this, EditingMode] { EquippedPolygonAction = EditingMode; } ),
+		FExecuteAction::CreateLambda( [this, EditingMode] { SetEquippedAction(EEditableMeshElementType::Polygon, EditingMode ); } ),
 		FCanExecuteAction::CreateLambda( [this] { return IsMeshElementTypeSelectedOrIsActiveSelectionMode( EEditableMeshElementType::Polygon ); } ),
 		FIsActionChecked::CreateLambda( [this, EditingMode] { return ( EquippedPolygonAction == EditingMode ); } )
 	) );
@@ -730,7 +764,7 @@ void FMeshEditorMode::Exit()
 	{
 		const FScopedTransaction Transaction( LOCTEXT( "UndoDeselectingAllMeshElements", "Deselect All Elements" ) );
 
-		UndoHelpers::StoreUndo( MeshEditorModeProxyObject, FDeselectAllMeshElementsChange( MoveTemp( FDeselectAllMeshElementsChangeInput() ) ).Execute( MeshEditorModeProxyObject ) );
+		DeselectAllMeshElements();
 	}
 
 	FToolkitManager::Get().CloseToolkit( Toolkit.ToSharedRef() );
@@ -813,8 +847,12 @@ UEditableMesh* FMeshEditorMode::FindOrCreateEditableMesh( UPrimitiveComponent& C
 
 void FMeshEditorMode::RollbackPreviewChanges()
 {
-	for( TTuple<UObject*, TUniquePtr<FChange>>& ObjectAndPreviewRevertChange : PreviewRevertChanges )
+	// NOTE: We iterate backwards here, because changes were added to our array in the order they originally
+	// happened.  But we'll need to apply their revert in the opposite order.
+	for( int32 ChangeIndex = PreviewRevertChanges.Num() - 1; ChangeIndex >= 0; --ChangeIndex )
 	{
+		TTuple<UObject*, TUniquePtr<FChange>>& ObjectAndPreviewRevertChange = PreviewRevertChanges[ ChangeIndex ];
+
 		UObject* Object = ObjectAndPreviewRevertChange.Get<0>();
 		TUniquePtr<FChange>& PreviewRevertChange = ObjectAndPreviewRevertChange.Get<1>();
 
@@ -1053,7 +1091,7 @@ void FMeshEditorMode::CommitEditableMeshIfNecessary( UEditableMesh* EditableMesh
 				ChangeInput.MeshElementsToSelect.Add( NewMeshElement );
 			}
 		}
-		UndoHelpers::StoreUndo( MeshEditorModeProxyObject, FSelectOrDeselectMeshElementsChange( MoveTemp( ChangeInput ) ).Execute( MeshEditorModeProxyObject ) );
+		TrackUndo( MeshEditorModeProxyObject, FSelectOrDeselectMeshElementsChange( MoveTemp( ChangeInput ) ).Execute( MeshEditorModeProxyObject ) );
 
 		FixUpMeshElements( SelectedVertices );
 		FixUpMeshElements( SelectedEdges );
@@ -1073,7 +1111,7 @@ void FMeshEditorMode::CommitEditableMeshIfNecessary( UEditableMesh* EditableMesh
 		{
 			for( auto& EdgeToSplit : SplitEdgeMeshAndEdgesToSplit.Value )
 			{
-				FixUpMeshElement( EdgeToSplit.Get<0>() );
+				FixUpMeshElement( EdgeToSplit );
 			}
 		}
 
@@ -1157,7 +1195,13 @@ void FMeshEditorMode::SelectMeshElements( const TArray<FMeshElement>& MeshElemen
 
 	CompoundRevertInput.Subchanges.Add( MoveTemp( FSelectOrDeselectMeshElementsChange( MoveTemp( ChangeInput ) ).Execute( MeshEditorModeProxyObject ) ) );
 
-	UndoHelpers::StoreUndo( MeshEditorModeProxyObject, MakeUnique<FCompoundChange>( MoveTemp( CompoundRevertInput ) ) );
+	TrackUndo( MeshEditorModeProxyObject, MakeUnique<FCompoundChange>( MoveTemp( CompoundRevertInput ) ) );
+}
+
+
+void FMeshEditorMode::DeselectAllMeshElements()
+{
+	TrackUndo( MeshEditorModeProxyObject, FDeselectAllMeshElementsChange( FDeselectAllMeshElementsChangeInput() ).Execute( MeshEditorModeProxyObject ) );
 }
 
 
@@ -1169,7 +1213,7 @@ void FMeshEditorMode::DeselectMeshElements( const TArray<FMeshElement>& MeshElem
 
 	CompoundRevertInput.Subchanges.Add( MoveTemp( FSelectOrDeselectMeshElementsChange( MoveTemp( ChangeInput ) ).Execute( MeshEditorModeProxyObject ) ) );
 
-	UndoHelpers::StoreUndo( MeshEditorModeProxyObject, MakeUnique<FCompoundChange>( MoveTemp( CompoundRevertInput ) ) );
+	TrackUndo( MeshEditorModeProxyObject, MakeUnique<FCompoundChange>( MoveTemp( CompoundRevertInput ) ) );
 }
 
 
@@ -1187,7 +1231,7 @@ void FMeshEditorMode::DeselectMeshElements( const TMap<UEditableMesh*, TArray<FM
 
 	CompoundRevertInput.Subchanges.Add( MoveTemp( FSelectOrDeselectMeshElementsChange( MoveTemp( ChangeInput ) ).Execute( MeshEditorModeProxyObject ) ) );
 
-	UndoHelpers::StoreUndo( MeshEditorModeProxyObject, MakeUnique<FCompoundChange>( MoveTemp( CompoundRevertInput ) ) );
+	TrackUndo( MeshEditorModeProxyObject, MakeUnique<FCompoundChange>( MoveTemp( CompoundRevertInput ) ) );
 }
 
 
@@ -1261,7 +1305,7 @@ bool FMeshEditorMode::DeleteSelectedMeshElement()
 
 		EditableMesh->EndModification();
 
-		UndoHelpers::StoreUndo( EditableMesh, EditableMesh->MakeUndo() );
+		TrackUndo( EditableMesh, EditableMesh->MakeUndo() );
 	}
 
 	// Make sure we're not still hovering over a polygon we deleted
@@ -1318,7 +1362,7 @@ void FMeshEditorMode::AddOrRemoveSubdivisionLevel( const bool bShouldAdd )
 
 			EditableMesh->EndModification();
 
-			UndoHelpers::StoreUndo( EditableMesh, EditableMesh->MakeUndo() );
+			TrackUndo( EditableMesh, EditableMesh->MakeUndo() );
 		}
 	}
 }
@@ -1353,9 +1397,7 @@ void FMeshEditorMode::QuadrangulateMesh()
 			}
 		}
 
-		FCompoundChangeInput CompoundRevertInput;
-		CompoundRevertInput.Subchanges.Add( MoveTemp( FDeselectAllMeshElementsChange( FDeselectAllMeshElementsChangeInput() ).Execute( MeshEditorModeProxyObject ) ) );
-		UndoHelpers::StoreUndo( MeshEditorModeProxyObject, MakeUnique<FCompoundChange>( MoveTemp( CompoundRevertInput ) ) );
+		DeselectAllMeshElements();
 
 		for( UEditableMesh* EditableMesh : SelectedMeshes )
 		{
@@ -1364,7 +1406,7 @@ void FMeshEditorMode::QuadrangulateMesh()
 			EditableMesh->QuadrangulateMesh( NewPolygonRefs );
 			EditableMesh->EndModification();
 
-			UndoHelpers::StoreUndo( EditableMesh, EditableMesh->MakeUndo() );
+			TrackUndo( EditableMesh, EditableMesh->MakeUndo() );
 		}
 	}
 }
@@ -1378,19 +1420,21 @@ void FMeshEditorMode::FrameSelectedElements( FEditorViewportClient* ViewportClie
 	{
 		case EEditableMeshElementType::Vertex:
 		{
-			static TMap<UEditableMesh*, TArray<TTuple<FMeshElement, FVertexID>>> SelectedMeshesAndVertices;
+			static TMap<UEditableMesh*, TArray<FMeshElement>> SelectedMeshesAndVertices;
 			GetSelectedMeshesAndVertices( SelectedMeshesAndVertices );
 			for( const auto& SelectedMeshAndVertices : SelectedMeshesAndVertices )
 			{
 				UEditableMesh* EditableMesh = SelectedMeshAndVertices.Key;
-				const TArray<TTuple<FMeshElement, FVertexID>>& VertexElements = SelectedMeshAndVertices.Value;
+				const TArray<FMeshElement>& VertexElements = SelectedMeshAndVertices.Value;
 
 				for( const auto& VertexElement : VertexElements )
 				{
-					const UPrimitiveComponent* Component = VertexElement.Get<0>().Component.Get();
+					const FVertexID VertexID( VertexElement.ElementAddress.ElementID );
+
+					const UPrimitiveComponent* Component = VertexElement.Component.Get();
 					if( Component )
 					{
-						const FVector VertexPosition = EditableMesh->GetVertexAttribute( VertexElement.Get<1>(), UEditableMeshAttribute::VertexPosition(), 0 );
+						const FVector VertexPosition = EditableMesh->GetVertexAttribute( VertexID, UEditableMeshAttribute::VertexPosition(), 0 );
 						BoundingBox += Component->GetComponentTransform().TransformPosition( VertexPosition );
 					}
 				}
@@ -1400,23 +1444,25 @@ void FMeshEditorMode::FrameSelectedElements( FEditorViewportClient* ViewportClie
 
 		case EEditableMeshElementType::Edge:
 		{
-			static TMap<UEditableMesh*, TArray<TTuple<FMeshElement, FEdgeID>>> SelectedMeshesAndEdges;
+			static TMap<UEditableMesh*, TArray<FMeshElement>> SelectedMeshesAndEdges;
 			GetSelectedMeshesAndEdges( SelectedMeshesAndEdges );
 			for( const auto& SelectedMeshAndEdges : SelectedMeshesAndEdges )
 			{
 				UEditableMesh* EditableMesh = SelectedMeshAndEdges.Key;
-				const TArray<TTuple<FMeshElement, FEdgeID>>& EdgeElements = SelectedMeshAndEdges.Value;
+				const TArray<FMeshElement>& EdgeElements = SelectedMeshAndEdges.Value;
 
 				for( const auto& EdgeElement : EdgeElements )
 				{
-					const UPrimitiveComponent* Component = EdgeElement.Get<0>().Component.Get();
+					const FEdgeID EdgeID( EdgeElement.ElementAddress.ElementID );
+
+					const UPrimitiveComponent* Component = EdgeElement.Component.Get();
 					if( Component )
 					{
-						const FVertexID VertexID0 = EditableMesh->GetEdgeVertex( EdgeElement.Get<1>(), 0 );
+						const FVertexID VertexID0 = EditableMesh->GetEdgeVertex( EdgeID, 0 );
 						const FVector VertexPosition0 = EditableMesh->GetVertexAttribute( VertexID0, UEditableMeshAttribute::VertexPosition(), 0 );
 						BoundingBox += Component->GetComponentTransform().TransformPosition( VertexPosition0 );
 
-						const FVertexID VertexID1 = EditableMesh->GetEdgeVertex( EdgeElement.Get<1>(), 1 );
+						const FVertexID VertexID1 = EditableMesh->GetEdgeVertex( EdgeID, 1 );
 						const FVector VertexPosition1 = EditableMesh->GetVertexAttribute( VertexID1, UEditableMeshAttribute::VertexPosition(), 0 );
 						BoundingBox += Component->GetComponentTransform().TransformPosition( VertexPosition1 );
 					}
@@ -1427,22 +1473,24 @@ void FMeshEditorMode::FrameSelectedElements( FEditorViewportClient* ViewportClie
 
 		case EEditableMeshElementType::Polygon:
 		{
-			static TMap<UEditableMesh*, TArray<TTuple<FMeshElement, FPolygonRef>>> SelectedMeshesAndPolygons;
+			static TMap<UEditableMesh*, TArray<FMeshElement>> SelectedMeshesAndPolygons;
 			GetSelectedMeshesAndPolygons( SelectedMeshesAndPolygons );
 			for( const auto& SelectedMeshAndPolygons : SelectedMeshesAndPolygons )
 			{
 				UEditableMesh* EditableMesh = SelectedMeshAndPolygons.Key;
-				const TArray<TTuple<FMeshElement, FPolygonRef>>& PolygonElements = SelectedMeshAndPolygons.Value;
+				const TArray<FMeshElement>& PolygonElements = SelectedMeshAndPolygons.Value;
 
 				for( const auto& PolygonElement : PolygonElements )
 				{
-					const UPrimitiveComponent* Component = PolygonElement.Get<0>().Component.Get();
+					const UPrimitiveComponent* Component = PolygonElement.Component.Get();
 					if( Component )
 					{
-						const int32 PolygonVertexCount = EditableMesh->GetPolygonPerimeterVertexCount( PolygonElement.Get<1>() );
+						const FPolygonRef PolygonRef( PolygonElement.ElementAddress.SectionID, FPolygonID( PolygonElement.ElementAddress.ElementID ) );
+
+						const int32 PolygonVertexCount = EditableMesh->GetPolygonPerimeterVertexCount( PolygonRef );
 						for( int32 PolygonVertexIndex = 0; PolygonVertexIndex < PolygonVertexCount; ++PolygonVertexIndex )
 						{
-							const FVector VertexPosition = EditableMesh->GetPolygonPerimeterVertexAttribute( PolygonElement.Get<1>(), PolygonVertexIndex, UEditableMeshAttribute::VertexPosition(), 0 );
+							const FVector VertexPosition = EditableMesh->GetPolygonPerimeterVertexAttribute( PolygonRef, PolygonVertexIndex, UEditableMeshAttribute::VertexPosition(), 0 );
 							BoundingBox += Component->GetComponentTransform().TransformPosition( VertexPosition );
 						}
 					}
@@ -1533,21 +1581,15 @@ bool FMeshEditorMode::RemoveSelectedEdges()
 
 		EditableMesh->EndModification();
 
-		UndoHelpers::StoreUndo( EditableMesh, EditableMesh->MakeUndo() );
+		TrackUndo( EditableMesh, EditableMesh->MakeUndo() );
 	}
 
 	{
 		// Make sure we're not still hovering over a polygon we removed
 		ClearInvalidSelectedElements();
 
-		FCompoundChangeInput CompoundRevertInput;
-
 		// Select the polygon leftover after removing the edge
-		FSelectOrDeselectMeshElementsChangeInput ChangeInput;
-		ChangeInput.MeshElementsToSelect = MeshElementsToSelect;
-		CompoundRevertInput.Subchanges.Add( MoveTemp( FSelectOrDeselectMeshElementsChange( MoveTemp( ChangeInput ) ).Execute( MeshEditorModeProxyObject ) ) );
-
-		UndoHelpers::StoreUndo( MeshEditorModeProxyObject, MakeUnique<FCompoundChange>( MoveTemp( CompoundRevertInput ) ) );
+		SelectMeshElements( MeshElementsToSelect );
 	}
 
 	return true;
@@ -1599,16 +1641,8 @@ bool FMeshEditorMode::SelectEdgeLoops()
 		}
 	}
 
-	{
-		FCompoundChangeInput CompoundRevertInput;
-		CompoundRevertInput.Subchanges.Add( MoveTemp( FDeselectAllMeshElementsChange( FDeselectAllMeshElementsChangeInput() ).Execute( MeshEditorModeProxyObject ) ) );
-
-		FSelectOrDeselectMeshElementsChangeInput ChangeInput;
-		ChangeInput.MeshElementsToSelect = MeshElementsToSelect;
-		CompoundRevertInput.Subchanges.Add( MoveTemp( FSelectOrDeselectMeshElementsChange( MoveTemp( ChangeInput ) ).Execute( MeshEditorModeProxyObject ) ) );
-
-		UndoHelpers::StoreUndo( MeshEditorModeProxyObject, MakeUnique<FCompoundChange>( MoveTemp( CompoundRevertInput ) ) );
-	}
+	DeselectAllMeshElements();
+	SelectMeshElements( MeshElementsToSelect );
 
 	return true;
 }
@@ -1693,7 +1727,7 @@ bool FMeshEditorMode::RemoveSelectedVertices()
 
 		EditableMesh->EndModification();
 
-		UndoHelpers::StoreUndo( EditableMesh, EditableMesh->MakeUndo() );
+		TrackUndo( EditableMesh, EditableMesh->MakeUndo() );
 	}
 
 	{
@@ -1707,7 +1741,7 @@ bool FMeshEditorMode::RemoveSelectedVertices()
 		ChangeInput.MeshElementsToSelect = MeshElementsToSelect;
 		CompoundRevertInput.Subchanges.Add( MoveTemp( FSelectOrDeselectMeshElementsChange( MoveTemp( ChangeInput ) ).Execute( MeshEditorModeProxyObject ) ) );
 
-		UndoHelpers::StoreUndo( MeshEditorModeProxyObject, MakeUnique<FCompoundChange>( MoveTemp( CompoundRevertInput ) ) );
+		TrackUndo( MeshEditorModeProxyObject, MakeUnique<FCompoundChange>( MoveTemp( CompoundRevertInput ) ) );
 	}
 
 	return true;
@@ -1787,7 +1821,7 @@ bool FMeshEditorMode::WeldSelectedVertices()
 
 		EditableMesh->EndModification();
 
-		UndoHelpers::StoreUndo( EditableMesh, EditableMesh->MakeUndo() );
+		TrackUndo( EditableMesh, EditableMesh->MakeUndo() );
 	}
 
 	{
@@ -1801,7 +1835,7 @@ bool FMeshEditorMode::WeldSelectedVertices()
 		ChangeInput.MeshElementsToSelect = MeshElementsToSelect;
 		CompoundRevertInput.Subchanges.Add( MoveTemp( FSelectOrDeselectMeshElementsChange( MoveTemp( ChangeInput ) ).Execute( MeshEditorModeProxyObject ) ) );
 
-		UndoHelpers::StoreUndo( MeshEditorModeProxyObject, MakeUnique<FCompoundChange>( MoveTemp( CompoundRevertInput ) ) );
+		TrackUndo( MeshEditorModeProxyObject, MakeUnique<FCompoundChange>( MoveTemp( CompoundRevertInput ) ) );
 	}
 
 	return true;
@@ -1851,7 +1885,7 @@ bool FMeshEditorMode::FlipSelectedPolygons()
 
 		EditableMesh->EndModification();
 
-		UndoHelpers::StoreUndo( EditableMesh, EditableMesh->MakeUndo() );
+		TrackUndo( EditableMesh, EditableMesh->MakeUndo() );
 	}
 
 	return true;
@@ -1932,7 +1966,7 @@ bool FMeshEditorMode::TriangulateSelectedPolygons()
 
 		EditableMesh->EndModification();
 
-		UndoHelpers::StoreUndo( EditableMesh, EditableMesh->MakeUndo() );
+		TrackUndo( EditableMesh, EditableMesh->MakeUndo() );
 	}
 
 	// Select the newly-created triangles
@@ -1946,7 +1980,7 @@ bool FMeshEditorMode::TriangulateSelectedPolygons()
 		ChangeInput.MeshElementsToSelect = MeshElementsToSelect;
 		CompoundRevertInput.Subchanges.Add( MoveTemp( FSelectOrDeselectMeshElementsChange( MoveTemp( ChangeInput ) ).Execute( MeshEditorModeProxyObject ) ) );
 
-		UndoHelpers::StoreUndo( MeshEditorModeProxyObject, MakeUnique<FCompoundChange>( MoveTemp( CompoundRevertInput ) ) );
+		TrackUndo( MeshEditorModeProxyObject, MakeUnique<FCompoundChange>( MoveTemp( CompoundRevertInput ) ) );
 	}
 
 
@@ -2044,7 +2078,7 @@ bool FMeshEditorMode::AssignMaterialToSelectedPolygons( UMaterialInterface* Sele
 			}
 			EditableMesh->EndModification();
 
-			UndoHelpers::StoreUndo( EditableMesh, EditableMesh->MakeUndo() );
+			TrackUndo( EditableMesh, EditableMesh->MakeUndo() );
 		}
 
 		// Select the newly-created triangles
@@ -2058,7 +2092,7 @@ bool FMeshEditorMode::AssignMaterialToSelectedPolygons( UMaterialInterface* Sele
 			ChangeInput.MeshElementsToSelect = MeshElementsToSelect;
 			CompoundRevertInput.Subchanges.Add( MoveTemp( FSelectOrDeselectMeshElementsChange( MoveTemp( ChangeInput ) ).Execute( MeshEditorModeProxyObject ) ) );
 
-			UndoHelpers::StoreUndo( MeshEditorModeProxyObject, MakeUnique<FCompoundChange>( MoveTemp( CompoundRevertInput ) ) );
+			TrackUndo( MeshEditorModeProxyObject, MakeUnique<FCompoundChange>( MoveTemp( CompoundRevertInput ) ) );
 		}
 
 		return true;
@@ -2076,7 +2110,7 @@ bool FMeshEditorMode::MakeSelectedEdgesHardOrSoft(const bool bMakeEdgesHard)
 	}
 
 	// Make edges hard or soft
-	static TMap< UEditableMesh*, TArray<TTuple<FMeshElement, FEdgeID>> > MeshesAndEdges;
+	static TMap< UEditableMesh*, TArray<FMeshElement> > MeshesAndEdges;
 	GetSelectedMeshesAndEdges( /* Out */ MeshesAndEdges );
 
 	if( MeshesAndEdges.Num() == 0 )
@@ -2100,9 +2134,11 @@ bool FMeshEditorMode::MakeSelectedEdgesHardOrSoft(const bool bMakeEdgesHard)
 		static TArray<bool> EdgesNewIsHard;
 		EdgesNewIsHard.Reset();
 
-		for( TTuple<FMeshElement, FEdgeID>& EdgeElement : MeshAndEdges.Value )
+		for( FMeshElement& EdgeElement : MeshAndEdges.Value )
 		{
-			EdgeIDs.Add( EdgeElement.Get<1>() );
+			const FEdgeID EdgeID( EdgeElement.ElementAddress.ElementID );
+
+			EdgeIDs.Add( EdgeID );
 			EdgesNewIsHard.Add( bMakeEdgesHard );
 		}
 
@@ -2116,7 +2152,7 @@ bool FMeshEditorMode::MakeSelectedEdgesHardOrSoft(const bool bMakeEdgesHard)
 
 		EditableMesh->EndModification();
 
-		UndoHelpers::StoreUndo( EditableMesh, EditableMesh->MakeUndo() );
+		TrackUndo( EditableMesh, EditableMesh->MakeUndo() );
 	}
 
 	return true;
@@ -2164,6 +2200,8 @@ void FMeshEditorMode::AddReferencedObjects( FReferenceCollector& Collector )
 		Collector.AddReferencedObject( Object );
 	}
 
+	Collector.AddReferencedObjects( ActiveActionModifiedMeshes );
+
 	Collector.AddReferencedObject( MeshEditorModeProxyObject );
 	Collector.AddReferencedObject( AssetContainer );
 }
@@ -3069,8 +3107,7 @@ void FMeshEditorMode::OnViewportInteractionInputUnhandled( FEditorViewportClient
 			if( SelectedMeshElements.Num() > 0 )
 			{
 				const FScopedTransaction Transaction( LOCTEXT( "UndoDeselectingAllMeshElements", "Deselect All Elements" ) );
-
-				UndoHelpers::StoreUndo( MeshEditorModeProxyObject, FDeselectAllMeshElementsChange( MoveTemp( FDeselectAllMeshElementsChangeInput() ) ).Execute( MeshEditorModeProxyObject ) );
+				DeselectAllMeshElements();
 			}
 		}
 	}
@@ -3284,40 +3321,16 @@ void FMeshEditorMode::UpdateActiveAction( const bool bIsActionFinishing )
 	};
 
 
-	check( !UndoHelpers::IsUndoSystemAvailable() || GEditor->IsTransactionActive() );
-
-	static TSet<UEditableMesh*> MeshesBeingModified;
-	MeshesBeingModified.Reset();
-
-
-	// NOTE: These are stored in the reverse order they need to be applied in!
-	static TArray< TTuple< UObject*, TUniquePtr<FChange>> > UpdateRevertChanges;
-	UpdateRevertChanges.Reset();
-
-	TSet<UEditableMesh*>* MeshesBeingModifiedPtr = &MeshesBeingModified;
-	TArray< TTuple< UObject*, TUniquePtr<FChange>>>* UpdateRevertChangesPtr = &UpdateRevertChanges;
-
-	auto TrackUndo = [MeshesBeingModifiedPtr, UpdateRevertChangesPtr]( UObject* Object, TUniquePtr<FChange> RevertChange )
-	{
-		if( RevertChange.IsValid() )
-		{
-			check( Object != nullptr );
+	// If this is an interim change, then everything that happens here -- all changes to our meshes, and even selection
+	// changes -- are guaranteed to be rolled back at the beginning of the next frame.  So we'll intercept any requests
+	// to store Undo history, and instead store those in a separate array to be processed ourselves next frame.
+	const EMeshModificationType MeshModificationType = bIsActionFinishing ? EMeshModificationType::Final : ( bIsFirstActiveActionUpdate ? EMeshModificationType::FirstInterim : EMeshModificationType::Interim );
+	this->bIsCapturingUndoForPreview = ( MeshModificationType != EMeshModificationType::Final );
+	check( GUndo == nullptr || GEditor->IsTransactionActive() );
+	ActiveActionModifiedMeshes.Reset();
 
-			// If the object is a mesh, make sure we've started to modify it
-			UEditableMesh* EditableMesh = Cast<UEditableMesh>( Object );
-			if( EditableMesh != nullptr )
-			{
-				// StartModification() must have already been called, otherwise it's too late at this point -- the mesh has been changed
-				// while the render thread was still using it.  Bad things will happen.  So we assert here.
-				check( MeshesBeingModifiedPtr->Contains( EditableMesh ) );
-			}
-
-			UpdateRevertChangesPtr->Add( TTuple<UObject*, TUniquePtr<FChange>>( Object, MoveTemp( RevertChange ) ) );
-		}
-	};
 
 	// Make sure StartModification() is called on all selected meshes
-	const EMeshModificationType MeshModificationType = bIsActionFinishing ? EMeshModificationType::Final : ( bIsFirstActiveActionUpdate ? EMeshModificationType::FirstInterim : EMeshModificationType::Interim );
 	const EMeshTopologyChange MeshTopologyChange = ( ( ActiveAction == EMeshEditAction::Move || ActiveAction == EMeshEditAction::MoveUsingGizmo ) ? EMeshTopologyChange::NoTopologyChange : EMeshTopologyChange::TopologyChange );
 	{
 		for( const FMeshElement& SelectedMeshElement : SelectedMeshElements )
@@ -3330,9 +3343,9 @@ void FMeshEditorMode::UpdateActiveAction( const bool bIsActionFinishing )
 					// @todo mesheditor debug
 					//GWarn->Logf( TEXT( "Selected: %s, Element: %s" ), *SelectedMeshElement.Component->GetName(), *SelectedMeshElement.ToString() );
 
-					if( !MeshesBeingModified.Contains( EditableMesh ) )	// @todo gizmo: All transformables will also hit this right?
+					if( !ActiveActionModifiedMeshes.Contains( EditableMesh ) )	// @todo gizmo: All transformables will also hit this right?
 					{
-						MeshesBeingModified.Add( EditableMesh );
+						ActiveActionModifiedMeshes.Add( EditableMesh );
 
 						EditableMesh->StartModification( MeshModificationType, MeshTopologyChange );
 
@@ -3353,7 +3366,7 @@ void FMeshEditorMode::UpdateActiveAction( const bool bIsActionFinishing )
 	else if( ActiveAction == EMeshEditAction::SplitEdge )
 	{
 		// Split edge
-		static TMap< UEditableMesh*, TArray< TTuple< FMeshElement, FEdgeID > > > SelectedMeshesAndEdges;
+		static TMap< UEditableMesh*, TArray<FMeshElement> > SelectedMeshesAndEdges;
 		GetSelectedMeshesAndEdges( /* Out */ SelectedMeshesAndEdges );
 
 		if( SelectedMeshesAndEdges.Num() > 0 )
@@ -3363,11 +3376,11 @@ void FMeshEditorMode::UpdateActiveAction( const bool bIsActionFinishing )
 				UEditableMesh* EditableMesh = MeshAndEdges.Key;
 				verify( !EditableMesh->AnyChangesToUndo() );
 
-				const TArray<TTuple<FMeshElement, FEdgeID>>& Edges = MeshAndEdges.Value;
+				const TArray<FMeshElement>& EdgeElements = MeshAndEdges.Value;
 
 				// Figure out where to split
 				static TArray<float> Splits;
-				FindEdgeSplitUnderInteractor( ActiveActionInteractor, EditableMesh, Edges, /* Out */ Splits );
+				FindEdgeSplitUnderInteractor( ActiveActionInteractor, EditableMesh, EdgeElements, /* Out */ Splits );
 
 				// Split the edges
 				if( Splits.Num() > 0 )
@@ -3385,10 +3398,9 @@ void FMeshEditorMode::UpdateActiveAction( const bool bIsActionFinishing )
 						}
 					}
 
-					for( const TTuple<FMeshElement, FEdgeID>& EdgeElementAndID : Edges )
+					for( const FMeshElement& EdgeMeshElement : EdgeElements )
 					{
-						const FMeshElement& EdgeMeshElement = EdgeElementAndID.Get<0>();
-						const FEdgeID EdgeID = EdgeElementAndID.Get<1>();
+						const FEdgeID EdgeID( EdgeMeshElement.ElementAddress.ElementID );
 
 						static TArray<FVertexID> NewVertexIDs;
 						NewVertexIDs.Reset();
@@ -3429,12 +3441,11 @@ void FMeshEditorMode::UpdateActiveAction( const bool bIsActionFinishing )
 				UEditableMesh* EditableMesh = MeshAndEdges.Key;
 				verify( !EditableMesh->AnyChangesToUndo() );
 
-				const TArray<TTuple<FMeshElement, FEdgeID>>& Edges = MeshAndEdges.Value;
+				const TArray<FMeshElement>& EdgeElements = MeshAndEdges.Value;
 
-				for( const TTuple<FMeshElement, FEdgeID>& EdgeElementAndID : Edges )
+				for( const FMeshElement& EdgeElement : EdgeElements )
 				{
-					const FMeshElement& EdgeMeshElement = EdgeElementAndID.Get<0>();
-					const FEdgeID EdgeID = EdgeElementAndID.Get<1>();
+					const FEdgeID EdgeID( EdgeElement.ElementAddress.ElementID );
 
 					static TArray<FVertexID> NewVertexIDs;
 					NewVertexIDs.Reset();
@@ -3448,8 +3459,8 @@ void FMeshEditorMode::UpdateActiveAction( const bool bIsActionFinishing )
 						{
 							FMeshElement MeshElementToSelect;
 							{
-								MeshElementToSelect.Component = EdgeMeshElement.Component;
-								MeshElementToSelect.ElementAddress.SubMeshAddress = EdgeMeshElement.ElementAddress.SubMeshAddress;
+								MeshElementToSelect.Component = EdgeElement.Component;
+								MeshElementToSelect.ElementAddress.SubMeshAddress = EdgeElement.ElementAddress.SubMeshAddress;
 								MeshElementToSelect.ElementAddress.ElementType = EEditableMeshElementType::Vertex;
 								MeshElementToSelect.ElementAddress.ElementID = NewVertexID;
 							}
@@ -3472,105 +3483,11 @@ void FMeshEditorMode::UpdateActiveAction( const bool bIsActionFinishing )
 		}
 	}
 
-	else if( ActiveAction == EMeshEditAction::InsertEdgeLoop )
-	{
-		// Insert edge loop
-		static TMap< UEditableMesh*, TArray< TTuple< FMeshElement, FEdgeID > > > SelectedMeshesAndEdges;
-		GetSelectedMeshesAndEdges( /* Out */ SelectedMeshesAndEdges );
-
-		if( SelectedMeshesAndEdges.Num() > 0 )
-		{
-			// Deselect the edges first, since they'll be deleted or split up while inserting the edge loop,
-			// and we want them to be re-selected after undo
-			{
-				FCompoundChangeInput CompoundRevertInput;
-
-				FSelectOrDeselectMeshElementsChangeInput ChangeInput;
-				for( auto& MeshAndEdges : SelectedMeshesAndEdges )
-				{
-					UEditableMesh* EditableMesh = MeshAndEdges.Key;
-					const TArray<TTuple<FMeshElement, FEdgeID>>& Edges = MeshAndEdges.Value;
-
-					for( const TTuple<FMeshElement, FEdgeID>& EdgeElementAndID : Edges )
-					{
-						const FMeshElement& EdgeMeshElement = EdgeElementAndID.Get<0>();
-						ChangeInput.MeshElementsToDeselect.Add( EdgeMeshElement );
-					}
-				}
-
-				CompoundRevertInput.Subchanges.Add( MoveTemp( FSelectOrDeselectMeshElementsChange( MoveTemp( ChangeInput ) ).Execute( MeshEditorModeProxyObject ) ) );
-
-				TrackUndo( MeshEditorModeProxyObject, MakeUnique<FCompoundChange>( MoveTemp( CompoundRevertInput ) ) );
-			}
-
-
-			for( auto& MeshAndEdges : SelectedMeshesAndEdges )
-			{
-				UEditableMesh* EditableMesh = MeshAndEdges.Key;
-				verify( !EditableMesh->AnyChangesToUndo() );
-
-				const TArray<TTuple<FMeshElement, FEdgeID>>& Edges = MeshAndEdges.Value;
-
-				// Figure out where to add the loop along the edge
-				static TArray<float> Splits;
-				FindEdgeSplitUnderInteractor( ActiveActionInteractor, EditableMesh, Edges, /* Out */ Splits );
-
-				// Insert the edge loop
-				if( Splits.Num() > 0 )
-				{
-					// @todo mesheditor edgeloop: Test code
-					if( false )
-					{
-						if( Splits[ 0 ] > 0.25f )
-						{
-							Splits.Insert( FMath::Max( Splits[ 0 ] - 0.2f, 0.0f ), 0 );
-						}
-						if( Splits.Last() < 0.75f )
-						{
-							Splits.Add( FMath::Min( Splits.Last() + 0.2f, 1.0f ) );
-						}
-					}
-
-					for( const TTuple<FMeshElement, FEdgeID>& EdgeElementAndID : Edges )
-					{
-						const FMeshElement& EdgeMeshElement = EdgeElementAndID.Get<0>();
-						const FEdgeID EdgeID = EdgeElementAndID.Get<1>();
-
-						static TArray<FEdgeID> NewEdgeIDs;
-						NewEdgeIDs.Reset();
-
-						EditableMesh->InsertEdgeLoop( EdgeID, Splits, /* Out */ NewEdgeIDs );
-
-						// Select all of the new edges that were created by inserting the loop
-						{
-							bShouldDeselectAllFirst = true;	// Don't keep the edge selected
-							for( const FEdgeID NewEdgeID : NewEdgeIDs )
-							{
-								FMeshElement MeshElementToSelect;
-								{
-									MeshElementToSelect.Component = EdgeMeshElement.Component;
-									MeshElementToSelect.ElementAddress.SubMeshAddress = EdgeMeshElement.ElementAddress.SubMeshAddress;
-									MeshElementToSelect.ElementAddress.ElementType = EEditableMeshElementType::Edge;
-									MeshElementToSelect.ElementAddress.ElementID = NewEdgeID;
-								}
-
-								// Queue selection of this new element.  We don't want it to be part of the current action.
-								MeshElementsToSelect.Add( MeshElementToSelect );
-							}
-						}
-					}
-
-					TrackUndo( EditableMesh, EditableMesh->MakeUndo() );
-				}
-			}
-		}
-	}
-
 	else if( ActiveAction == EMeshEditAction::ExtrudePolygon ||
 			 ActiveAction == EMeshEditAction::FreelyExtrudePolygon )
 	{
 		// Extrude polygon
-		static TMap< UEditableMesh*, TArray< TTuple< FMeshElement, FPolygonRef > > > MeshesWithPolygonsToExtrude;
+		static TMap< UEditableMesh*, TArray< FMeshElement > > MeshesWithPolygonsToExtrude;
 		GetSelectedMeshesAndPolygons( /* Out */ MeshesWithPolygonsToExtrude );
 
 		if( MeshesWithPolygonsToExtrude.Num() > 0 )
@@ -3589,11 +3506,10 @@ void FMeshEditorMode::UpdateActiveAction( const bool bIsActionFinishing )
 				for( auto& MeshAndPolygons : MeshesWithPolygonsToExtrude )
 				{
 					UEditableMesh* EditableMesh = MeshAndPolygons.Key;
-					const TArray<TTuple<FMeshElement, FPolygonRef>>& PolygonsToExtrude = MeshAndPolygons.Value;
+					const TArray<FMeshElement>& PolygonsToExtrude = MeshAndPolygons.Value;
 
-					for( const TTuple<FMeshElement, FPolygonRef>& PolygonToExtrude : PolygonsToExtrude )
+					for( const FMeshElement& PolygonElement : PolygonsToExtrude )
 					{
-						const FMeshElement& PolygonElement = PolygonToExtrude.Get<0>();
 						ChangeInput.MeshElementsToDeselect.Add( PolygonElement );
 					}
 				}
@@ -3606,13 +3522,14 @@ void FMeshEditorMode::UpdateActiveAction( const bool bIsActionFinishing )
 			for( auto& MeshAndPolygons : MeshesWithPolygonsToExtrude )
 			{
 				UEditableMesh* EditableMesh = MeshAndPolygons.Key;
-				const TArray<TTuple<FMeshElement, FPolygonRef>>& PolygonsToExtrude = MeshAndPolygons.Value;
+				const TArray<FMeshElement>& PolygonsToExtrude = MeshAndPolygons.Value;
 
 				static TArray<FPolygonRef> PolygonRefsToExtrude;
 				PolygonRefsToExtrude.Reset();
-				for( const TTuple<FMeshElement, FPolygonRef>& PolygonToExtrude : PolygonsToExtrude )
+				for( const FMeshElement& PolygonToExtrude : PolygonsToExtrude )
 				{
-					PolygonRefsToExtrude.Add( PolygonToExtrude.Get<1>() );
+					FPolygonRef PolygonRef( PolygonToExtrude.ElementAddress.SectionID, FPolygonID( PolygonToExtrude.ElementAddress.ElementID ) );
+					PolygonRefsToExtrude.Add( PolygonRef );
 				}
 
 
@@ -3637,7 +3554,7 @@ void FMeshEditorMode::UpdateActiveAction( const bool bIsActionFinishing )
 					verify( !EditableMesh->AnyChangesToUndo() );
 
 					// Position the new polygon to where the interactor is
-					UPrimitiveComponent* Component = PolygonsToExtrude[ 0 ].Get<0>().Component.Get();	// NOTE: All polygons in this array belong to the same mesh/component, so we just need the first element
+					UPrimitiveComponent* Component = PolygonsToExtrude[ 0 ].Component.Get();	// NOTE: All polygons in this array belong to the same mesh/component, so we just need the first element
 					check( Component != nullptr );
 					// @todo mesheditor: We're working with a float here, so we'll treat the component scale as a scalar (X)
 					const float ComponentSpaceExtrudeDistance = ExtrudeDistance / Component->ComponentToWorld.GetScale3D().X;
@@ -3653,7 +3570,7 @@ void FMeshEditorMode::UpdateActiveAction( const bool bIsActionFinishing )
 					{
 						const FPolygonRef& NewExtrudedFrontPolygon = NewExtrudedFrontPolygons[ NewPolygonNumber ];
 
-						const FMeshElement& MeshElement = PolygonsToExtrude[ NewPolygonNumber ].Get<0>();
+						const FMeshElement& MeshElement = PolygonsToExtrude[ NewPolygonNumber ];
 
 						FMeshElement NewExtrudedPolygonMeshElement;
 						{
@@ -3675,7 +3592,7 @@ void FMeshEditorMode::UpdateActiveAction( const bool bIsActionFinishing )
 	else if( ActiveAction == EMeshEditAction::InsetPolygon ||
 			 ActiveAction == EMeshEditAction::BevelPolygon )
 	{
-		static TMap< UEditableMesh*, TArray< TTuple< FMeshElement, FPolygonRef > > > MeshesWithPolygonsToInset;
+		static TMap< UEditableMesh*, TArray< FMeshElement > > MeshesWithPolygonsToInset;
 		GetSelectedMeshesAndPolygons( /* Out */ MeshesWithPolygonsToInset );
 
 		if( MeshesWithPolygonsToInset.Num() > 0 )
@@ -3716,13 +3633,8 @@ void FMeshEditorMode::UpdateActiveAction( const bool bIsActionFinishing )
 							for( auto& MeshAndPolygons : MeshesWithPolygonsToInset )
 							{
 								UEditableMesh* EditableMesh = MeshAndPolygons.Key;
-								const TArray<TTuple<FMeshElement, FPolygonRef>>& PolygonsToDeselect = MeshAndPolygons.Value;
-
-								for( const TTuple<FMeshElement, FPolygonRef>& PolygonToDeselect : PolygonsToDeselect )
-								{
-									const FMeshElement& PolygonElement = PolygonToDeselect.Get<0>();
-									ChangeInput.MeshElementsToDeselect.Add( PolygonElement );
-								}
+								const TArray<FMeshElement>& PolygonsToDeselect = MeshAndPolygons.Value;
+								ChangeInput.MeshElementsToDeselect.Append( PolygonsToDeselect );
 							}
 
 							CompoundRevertInput.Subchanges.Add( MoveTemp( FSelectOrDeselectMeshElementsChange( MoveTemp( ChangeInput ) ).Execute( MeshEditorModeProxyObject ) ) );
@@ -3733,17 +3645,18 @@ void FMeshEditorMode::UpdateActiveAction( const bool bIsActionFinishing )
 						for( auto& MeshAndPolygons : MeshesWithPolygonsToInset )
 						{
 							UEditableMesh* EditableMesh = MeshAndPolygons.Key;
-							const TArray<TTuple<FMeshElement, FPolygonRef>>& PolygonsToInset = MeshAndPolygons.Value;
+							const TArray<FMeshElement>& PolygonsToInset = MeshAndPolygons.Value;
 
-							UPrimitiveComponent* Component = PolygonsToInset[ 0 ].Get<0>().Component.Get();	// NOTE: All polygons in this array belong to the same mesh/component, so we just need the first element
+							UPrimitiveComponent* Component = PolygonsToInset[ 0 ].Component.Get();	// NOTE: All polygons in this array belong to the same mesh/component, so we just need the first element
 							check( Component != nullptr );
 
 
 							static TArray<FPolygonRef> PolygonRefsToInset;
 							PolygonRefsToInset.Reset();
-							for( const TTuple<FMeshElement, FPolygonRef>& PolygonToInset : PolygonsToInset )
+							for( const FMeshElement& PolygonToInset : PolygonsToInset )
 							{
-								PolygonRefsToInset.Add( PolygonToInset.Get<1>() );
+								FPolygonRef PolygonRef( PolygonToInset.ElementAddress.SectionID, FPolygonID( PolygonToInset.ElementAddress.ElementID ) );
+								PolygonRefsToInset.Add( PolygonRef );
 							}
 
 							{
@@ -3771,7 +3684,7 @@ void FMeshEditorMode::UpdateActiveAction( const bool bIsActionFinishing )
 									{
 										const FPolygonRef& NewInsetCenterPolygon = NewCenterInsetPolygons[ NewPolygonNumber ];
 
-										const FMeshElement& MeshElement = PolygonsToInset[ NewPolygonNumber ].Get<0>();
+										const FMeshElement& MeshElement = PolygonsToInset[ NewPolygonNumber ];
 
 										FMeshElement PolygonMeshElement;
 										{
@@ -3791,7 +3704,7 @@ void FMeshEditorMode::UpdateActiveAction( const bool bIsActionFinishing )
 									{
 										const FPolygonRef& NewInsetSidePolygon = NewSideInsetPolygons[ NewPolygonNumber ];
 
-										const FMeshElement& MeshElement = PolygonsToInset[ NewPolygonNumber ].Get<0>();
+										const FMeshElement& MeshElement = PolygonsToInset[ NewPolygonNumber ];
 
 										FMeshElement PolygonMeshElement;
 										{
@@ -3817,7 +3730,7 @@ void FMeshEditorMode::UpdateActiveAction( const bool bIsActionFinishing )
 	else if( ActiveAction == EMeshEditAction::ExtendEdge )
 	{
 		// Extend edge
-		static TMap< UEditableMesh*, TArray< TTuple< FMeshElement, FEdgeID > > > MeshesWithEdgesToExtend;
+		static TMap< UEditableMesh*, TArray<FMeshElement> > MeshesWithEdgesToExtend;
 		GetSelectedMeshesAndEdges( /* Out */ MeshesWithEdgesToExtend );
 
 		if( MeshesWithEdgesToExtend.Num() > 0 )
@@ -3828,16 +3741,16 @@ void FMeshEditorMode::UpdateActiveAction( const bool bIsActionFinishing )
 			for( auto& MeshAndEdges : MeshesWithEdgesToExtend )
 			{
 				UEditableMesh* EditableMesh = MeshAndEdges.Key;
-				const TArray<TTuple<FMeshElement, FEdgeID>>& EdgeElementsToExtend = MeshAndEdges.Value;
+				const TArray<FMeshElement>& EdgeElementsToExtend = MeshAndEdges.Value;
 
 				static TArray<FEdgeID> EdgeIDsToExtend;
 				EdgeIDsToExtend.Reset();
 
-				for( const TTuple<FMeshElement, FEdgeID>& EdgeElementToExtend : EdgeElementsToExtend )
+				for( const FMeshElement& EdgeElementToExtend : EdgeElementsToExtend )
 				{
-					const FMeshElement& PolygonElement = EdgeElementToExtend.Get<0>();
+					const FEdgeID EdgeID( EdgeElementToExtend.ElementAddress.ElementID );
 
-					EdgeIDsToExtend.Add( EdgeElementToExtend.Get<1>() );
+					EdgeIDsToExtend.Add( EdgeID );
 				}
 
 
@@ -3857,7 +3770,7 @@ void FMeshEditorMode::UpdateActiveAction( const bool bIsActionFinishing )
 					{
 						const FEdgeID NewExtendedEdgeID = NewExtendedEdgeIDs[ NewEdgeNumber ];
 
-						const FMeshElement& MeshElement = EdgeElementsToExtend[ NewEdgeNumber ].Get<0>();
+						const FMeshElement& MeshElement = EdgeElementsToExtend[ NewEdgeNumber ];
 
 						FMeshElement NewExtendedEdgeMeshElement;
 						{
@@ -3878,7 +3791,7 @@ void FMeshEditorMode::UpdateActiveAction( const bool bIsActionFinishing )
 
 	else if( ActiveAction == EMeshEditAction::ExtendVertex )
 	{
-		static TMap< UEditableMesh*, TArray< TTuple< FMeshElement, FVertexID > > > MeshesWithVerticesToExtend;
+		static TMap< UEditableMesh*, TArray< FMeshElement > > MeshesWithVerticesToExtend;
 		GetSelectedMeshesAndVertices( /* Out */ MeshesWithVerticesToExtend );
 
 		if( MeshesWithVerticesToExtend.Num() > 0 )
@@ -3889,23 +3802,23 @@ void FMeshEditorMode::UpdateActiveAction( const bool bIsActionFinishing )
 			for( auto& MeshAndVertices : MeshesWithVerticesToExtend )
 			{
 				UEditableMesh* EditableMesh = MeshAndVertices.Key;
-				const TArray<TTuple<FMeshElement, FVertexID>>& VertexElementsToExtend = MeshAndVertices.Value;
+				const TArray<FMeshElement>& VertexElementsToExtend = MeshAndVertices.Value;
 
 				static TArray<FVertexID> VertexIDsToExtend;
 				VertexIDsToExtend.Reset();
 
-				for( const TTuple<FMeshElement, FVertexID>& VertexElementToExtend : VertexElementsToExtend )
+				for( const FMeshElement& VertexElementToExtend : VertexElementsToExtend )
 				{
-					const FMeshElement& PolygonElement = VertexElementToExtend.Get<0>();
+					const FVertexID VertexID( VertexElementToExtend.ElementAddress.ElementID );
 
-					VertexIDsToExtend.Add( VertexElementToExtend.Get<1>() );
+					VertexIDsToExtend.Add( VertexID );
 				}
 
 
 				verify( !EditableMesh->AnyChangesToUndo() );
 
 				{
-					UPrimitiveComponent* ComponentPtr = VertexElementsToExtend[ 0 ].Get<0>().Component.Get();
+					UPrimitiveComponent* ComponentPtr = VertexElementsToExtend[ 0 ].Component.Get();
 					check( ComponentPtr != nullptr );
 					UPrimitiveComponent& Component = *ComponentPtr;
 					const FMatrix ComponentToWorldMatrix = Component.GetRenderMatrix();
@@ -3925,7 +3838,7 @@ void FMeshEditorMode::UpdateActiveAction( const bool bIsActionFinishing )
 					{
 						const FVertexID NewExtendedVertexID = NewExtendedVertexIDs[ NewVertexNumber ];
 
-						const FMeshElement& MeshElement = VertexElementsToExtend[ NewVertexNumber ].Get<0>();
+						const FMeshElement& MeshElement = VertexElementsToExtend[ NewVertexNumber ];
 
 						FMeshElement NewExtendedVertexMeshElement;
 						{
@@ -3947,8 +3860,8 @@ void FMeshEditorMode::UpdateActiveAction( const bool bIsActionFinishing )
 	else if( ActiveAction == EMeshEditAction::EditVertexCornerSharpness ||
 			 ActiveAction == EMeshEditAction::EditEdgeCreaseSharpness )
 	{
-		static TMap< UEditableMesh*, TArray< TTuple< FMeshElement, FVertexID > > > MeshesWithSelectedVertices;
-		static TMap< UEditableMesh*, TArray< TTuple< FMeshElement, FEdgeID > > > MeshesWithSelectedEdges;
+		static TMap< UEditableMesh*, TArray<FMeshElement> > MeshesWithSelectedVertices;
+		static TMap< UEditableMesh*, TArray<FMeshElement> > MeshesWithSelectedEdges;
 		GetSelectedMeshesAndVertices( /* Out */ MeshesWithSelectedVertices );
 		GetSelectedMeshesAndEdges( /* Out */ MeshesWithSelectedEdges );
 
@@ -4007,7 +3920,7 @@ void FMeshEditorMode::UpdateActiveAction( const bool bIsActionFinishing )
 			for( auto& MeshAndSelectedVertices : MeshesWithSelectedVertices )
 			{
 				UEditableMesh* EditableMesh = MeshAndSelectedVertices.Key;
-				const TArray<TTuple<FMeshElement, FVertexID>>& VertexElementsAndIDs = MeshAndSelectedVertices.Value;
+				const TArray<FMeshElement>& VertexElements = MeshAndSelectedVertices.Value;
 
 				static TArray<FVertexID> VertexIDs;
 				VertexIDs.Reset();
@@ -4015,9 +3928,9 @@ void FMeshEditorMode::UpdateActiveAction( const bool bIsActionFinishing )
 				static TArray<float> NewSharpnessValues;
 				NewSharpnessValues.Reset();
 
-				for( const TTuple<FMeshElement, FVertexID>& VertexElementAndID : VertexElementsAndIDs )
+				for( const FMeshElement& VertexElement : VertexElements )
 				{
-					const FVertexID VertexID = VertexElementAndID.Get<1>();
+					const FVertexID VertexID( VertexElement.ElementAddress.ElementID );
 
 					const float CurrentSharpnessValue = EditableMesh->GetVertexAttribute( VertexID, UEditableMeshAttribute::VertexCornerSharpness(), 0 ).X;
 					const float NewSharpnessValue = FMath::Clamp( CurrentSharpnessValue + ScaledDragDelta, 0.0f, 1.0f );
@@ -4037,7 +3950,7 @@ void FMeshEditorMode::UpdateActiveAction( const bool bIsActionFinishing )
 			for( auto& MeshAndSelectedEdges : MeshesWithSelectedEdges )
 			{
 				UEditableMesh* EditableMesh = MeshAndSelectedEdges.Key;
-				const TArray<TTuple<FMeshElement, FEdgeID>>& EdgeElementsAndIDs = MeshAndSelectedEdges.Value;
+				const TArray<FMeshElement>& EdgeElements = MeshAndSelectedEdges.Value;
 
 				static TArray<FEdgeID> EdgeIDs;
 				EdgeIDs.Reset();
@@ -4045,9 +3958,9 @@ void FMeshEditorMode::UpdateActiveAction( const bool bIsActionFinishing )
 				static TArray<float> NewSharpnessValues;
 				NewSharpnessValues.Reset();
 
-				for( const TTuple<FMeshElement, FEdgeID>& EdgeElementAndID : EdgeElementsAndIDs )
+				for( const FMeshElement& EdgeElement : EdgeElements )
 				{
-					const FEdgeID EdgeID = EdgeElementAndID.Get<1>();
+					const FEdgeID EdgeID( EdgeElement.ElementAddress.ElementID );
 
 					const float CurrentSharpnessValue = EditableMesh->GetEdgeAttribute( EdgeID, UEditableMeshAttribute::EdgeCreaseSharpness(), 0 ).X;
 					const float NewSharpnessValue = FMath::Clamp( CurrentSharpnessValue + ScaledDragDelta, 0.0f, 1.0f );
@@ -4338,16 +4251,36 @@ void FMeshEditorMode::UpdateActiveAction( const bool bIsActionFinishing )
 			}
 		}
 	}
+	else
+	{
+		// Check for registered commands that are active right now
+		bool bFoundValidCommand = false;
+		for( TObjectIterator<UMeshEditorCommand> CommandCDOIter( RF_NoFlags ); CommandCDOIter; ++CommandCDOIter )
+		{
+			UMeshEditorCommand* CommandCDO = *CommandCDOIter;
+			if( !( CommandCDO->GetClass()->GetClassFlags() & CLASS_Abstract ) )
+			{
+				if( ActiveAction == CommandCDO->GetCommandName() )
+				{
+					CommandCDO->ApplyDuringDrag( *this, /* Out */ bShouldDeselectAllFirst, /* Out */ MeshElementsToSelect );
+
+					// Should always only be one candidate
+					bFoundValidCommand = true;
+					break;
+				}
+			}
+		}
+		check( bFoundValidCommand );	// There must have been a command registered to initiate this action
+	}
+
 
 	// Update selection state.  Note that we intentionally do this BEFORE we perform any dragging, so that
 	// we'll be dragging any newly-generated geometry from the mesh edit action.  For example, when extending
 	// an edge we want to drag around the newly-created edge, not the edge that was selected before.
 	{
-		FCompoundChangeInput CompoundRevertInput;
-
 		if( bShouldDeselectAllFirst )
 		{
-			CompoundRevertInput.Subchanges.Add( MoveTemp( FDeselectAllMeshElementsChange( FDeselectAllMeshElementsChangeInput() ).Execute( MeshEditorModeProxyObject ) ) );
+			DeselectAllMeshElements();
 		}
 		else
 		{
@@ -4357,14 +4290,7 @@ void FMeshEditorMode::UpdateActiveAction( const bool bIsActionFinishing )
 
 		if( MeshElementsToSelect.Num() > 0 )
 		{
-			FSelectOrDeselectMeshElementsChangeInput ChangeInput;
-			ChangeInput.MeshElementsToSelect = MeshElementsToSelect;
-			CompoundRevertInput.Subchanges.Add( MoveTemp( FSelectOrDeselectMeshElementsChange( MoveTemp( ChangeInput ) ).Execute( MeshEditorModeProxyObject ) ) );
-		}
-
-		if( CompoundRevertInput.Subchanges.Num() > 0 )
-		{
-			TrackUndo( MeshEditorModeProxyObject, MakeUnique<FCompoundChange>( MoveTemp( CompoundRevertInput ) ) );
+			SelectMeshElements( MeshElementsToSelect );
 		}
 	}
 
@@ -4509,7 +4435,7 @@ void FMeshEditorMode::UpdateActiveAction( const bool bIsActionFinishing )
 	// Make sure EndModification() is called on any EditableMesh objects that were modified, so their graphics and physics
 	// state is updated.
 	{
-		for( UEditableMesh* EditableMesh : MeshesBeingModified )
+		for( UEditableMesh* EditableMesh : ActiveActionModifiedMeshes )
 		{
 			if( EditableMesh != nullptr )
 			{
@@ -4521,43 +4447,9 @@ void FMeshEditorMode::UpdateActiveAction( const bool bIsActionFinishing )
 	}
 	
 
-	// Store final undo state for all changes that happened in this update
-	{
-		if( MeshModificationType == EMeshModificationType::Final )
-		{
-			// If we're finalizing the action, this will save the undo state for everything that happened in this function, including
-			// selection changes.  
-			for( TTuple<UObject*, TUniquePtr<FChange>>& UpdateRevertChange : UpdateRevertChanges )
-			{
-				UObject* Object = UpdateRevertChange.Get<0>();
-				check( Object != nullptr );
-
-				TUniquePtr<FChange>& RevertChange = UpdateRevertChange.Get<1>();
-				check( RevertChange.IsValid() );
-
-				UndoHelpers::StoreUndo( Object, MoveTemp( RevertChange ) );
-			}
-		}
-		else
-		{
-			// Otherwise, we'll store the commands to undo in our 'PreviewRevertChanges' member, so they can be
-			// rolled back at the beginning of the next frame before any new interactions take place.  This allows the user to preview
-			// (potentially highly destructive) changes live!
-			// NOTE: UpdateRevertChanges has entires stored in reverse order, so we'll iterate backwards as we build up or undo buffer
-			for( int32 ChangeIndex = UpdateRevertChanges.Num() - 1; ChangeIndex >= 0; --ChangeIndex )
-			{
-				TTuple<UObject*, TUniquePtr<FChange>>& UpdateRevertChange = UpdateRevertChanges[ ChangeIndex ];
-
-				PreviewRevertChanges.Add( MoveTemp( UpdateRevertChange ) );
-			}
-
-			PreviewTopologyChange = MeshTopologyChange;
-		}
-	}
-
 	// Reset temporary containers
-	UpdateRevertChanges.Reset();
-	MeshesBeingModified.Reset();
+	bIsCapturingUndoForPreview = false;
+	ActiveActionModifiedMeshes.Reset();
 
 	bIsFirstActiveActionUpdate = false;
 }
@@ -4586,76 +4478,7 @@ void FMeshEditorMode::GetSelectedMeshesAndElements( EEditableMeshElementType Ele
 }
 
 
-void FMeshEditorMode::GetSelectedMeshesAndVertices( TMap< UEditableMesh*, TArray< TTuple< FMeshElement, FVertexID > > >& OutMeshesAndVertices )
-{
-	OutMeshesAndVertices.Reset();
-
-	for( int32 SelectedElementIndex = 0; SelectedElementIndex < this->SelectedMeshElements.Num(); ++SelectedElementIndex )
-	{
-		const FMeshElement& SelectedMeshElement = this->SelectedMeshElements[ SelectedElementIndex ];
-		if( SelectedMeshElement.IsValidMeshElement() )
-		{
-			UEditableMesh* EditableMesh = this->FindOrCreateEditableMesh( *SelectedMeshElement.Component, SelectedMeshElement.ElementAddress.SubMeshAddress );
-			if( EditableMesh != nullptr )
-			{
-				if( SelectedMeshElement.ElementAddress.ElementType == EEditableMeshElementType::Vertex )
-				{
-					TArray< TTuple< FMeshElement, FVertexID > >& Vertices = OutMeshesAndVertices.FindOrAdd( EditableMesh );
-					Vertices.Add( TTuple< FMeshElement, FVertexID >( SelectedMeshElement, SelectedMeshElement.ElementAddress.ElementID ) );
-				}
-			}
-		}
-	}
-}
-
-
-void FMeshEditorMode::GetSelectedMeshesAndEdges( TMap< UEditableMesh*, TArray< TTuple< FMeshElement, FEdgeID > > >& OutMeshesAndEdges )
-{
-	OutMeshesAndEdges.Reset();
-
-	for( int32 SelectedElementIndex = 0; SelectedElementIndex < this->SelectedMeshElements.Num(); ++SelectedElementIndex )
-	{
-		const FMeshElement& SelectedMeshElement = this->SelectedMeshElements[ SelectedElementIndex ];
-		if( SelectedMeshElement.IsValidMeshElement() )
-		{
-			UEditableMesh* EditableMesh = this->FindOrCreateEditableMesh( *SelectedMeshElement.Component, SelectedMeshElement.ElementAddress.SubMeshAddress );
-			if( EditableMesh != nullptr )
-			{
-				if( SelectedMeshElement.ElementAddress.ElementType == EEditableMeshElementType::Edge )
-				{
-					TArray< TTuple< FMeshElement, FEdgeID > >& Edges = OutMeshesAndEdges.FindOrAdd( EditableMesh );
-					Edges.Add( TTuple< FMeshElement, FEdgeID >( SelectedMeshElement, SelectedMeshElement.ElementAddress.ElementID ) );
-				}
-			}
-		}
-	}
-}
-
-
-void FMeshEditorMode::GetSelectedMeshesAndPolygons(	TMap< UEditableMesh*, TArray< TTuple< FMeshElement, FPolygonRef > > >& OutMeshesAndPolygons )
-{
-	OutMeshesAndPolygons.Reset();
-
-	for( int32 SelectedElementIndex = 0; SelectedElementIndex < this->SelectedMeshElements.Num(); ++SelectedElementIndex )
-	{
-		const FMeshElement& SelectedMeshElement = this->SelectedMeshElements[ SelectedElementIndex ];
-		if( SelectedMeshElement.IsValidMeshElement() )
-		{
-			UEditableMesh* EditableMesh = this->FindOrCreateEditableMesh( *SelectedMeshElement.Component, SelectedMeshElement.ElementAddress.SubMeshAddress );
-			if( EditableMesh != nullptr )
-			{
-				if( SelectedMeshElement.ElementAddress.ElementType == EEditableMeshElementType::Polygon )
-				{
-					TArray< TTuple< FMeshElement, FPolygonRef > >& Polygons = OutMeshesAndPolygons.FindOrAdd( EditableMesh );
-					Polygons.Add( TTuple< FMeshElement, FPolygonRef >( SelectedMeshElement, FPolygonRef( SelectedMeshElement.ElementAddress.SectionID, FPolygonID( SelectedMeshElement.ElementAddress.ElementID ) ) ) );
-				}
-			}
-		}
-	}
-};
-
-
-void FMeshEditorMode::FindEdgeSplitUnderInteractor(	UViewportInteractor* ViewportInteractor, const UEditableMesh* EditableMesh, const TArray<TTuple<FMeshElement, FEdgeID>>& Edges, TArray<float>& OutSplits )
+void FMeshEditorMode::FindEdgeSplitUnderInteractor(	UViewportInteractor* ViewportInteractor, const UEditableMesh* EditableMesh, const TArray<FMeshElement>& EdgeElements, TArray<float>& OutSplits )
 {
 	check( ViewportInteractor != nullptr );
 
@@ -4665,10 +4488,9 @@ void FMeshEditorMode::FindEdgeSplitUnderInteractor(	UViewportInteractor* Viewpor
 	// selected edges, and choose a split offset based on the closest point along one of those edges
 	// to the interactor.  All selected edges will then be split by the same proportion.
 	float ClosestDistanceToEdge = MAX_flt;
-	for( const TTuple<FMeshElement, FEdgeID>& EdgeElementAndID : Edges )
+	for( const FMeshElement& EdgeMeshElement : EdgeElements )
 	{
-		const FMeshElement& EdgeMeshElement = EdgeElementAndID.Get<0>();
-		const FEdgeID EdgeID = EdgeElementAndID.Get<1>();
+		const FEdgeID EdgeID( EdgeMeshElement.ElementAddress.ElementID );
 
 		FMeshEditorInteractorData& MeshEditorInteractorData = this->GetMeshEditorInteractorData( ViewportInteractor );
 		if( MeshEditorInteractorData.bLaserIsValid || MeshEditorInteractorData.bGrabberSphereIsValid )
@@ -5092,7 +4914,7 @@ void FMeshEditorMode::SetMeshElementSelectionMode( EEditableMeshElementType Elem
 	const FScopedTransaction Transaction( LOCTEXT( "ChangeMeshElementSelectionMode", "Change Mesh Element Selection Mode" ) );
 	FSetElementSelectionModeChangeInput ChangeInput;
 	ChangeInput.Mode = ElementType;
-	UndoHelpers::StoreUndo( MeshEditorModeProxyObject, FSetElementSelectionModeChange( MoveTemp( ChangeInput ) ).Execute( MeshEditorModeProxyObject ) );
+	TrackUndo( MeshEditorModeProxyObject, FSetElementSelectionModeChange( MoveTemp( ChangeInput ) ).Execute( MeshEditorModeProxyObject ) );
 }
 
 
@@ -5159,7 +4981,7 @@ void FMeshEditorMode::ClearInvalidSelectedElements()
 {
 	// Make sure we're not still hovering over a polygon we removed
 	TUniquePtr<FChange> RevertChange = ClearInvalidSelectedElements_Internal();
-	UndoHelpers::StoreUndo( MeshEditorModeProxyObject, MoveTemp( RevertChange ) );
+	TrackUndo( MeshEditorModeProxyObject, MoveTemp( RevertChange ) );
 }
 
 
@@ -5404,10 +5226,10 @@ void FMeshEditorMode::OnViewportInteractionInputAction( FEditorViewportClient& V
 							for( auto& MeshAndEdges : SplitEdgeMeshesAndEdgesToSplit )
 							{
 								UEditableMesh* EditableMesh = MeshAndEdges.Key;
-								const TArray<TTuple<FMeshElement, FEdgeID>>& Edges = MeshAndEdges.Value;
+								const TArray<FMeshElement>& EdgeElements = MeshAndEdges.Value;
 
 								// Figure out where to split
-								FindEdgeSplitUnderInteractor( ViewportInteractor, EditableMesh, Edges, /* Out */ SplitEdgeSplitList );
+								FindEdgeSplitUnderInteractor( ViewportInteractor, EditableMesh, EdgeElements, /* Out */ SplitEdgeSplitList );
 
 								// Split the edges
 								if( SplitEdgeSplitList.Num() > 0 )
@@ -5422,13 +5244,6 @@ void FMeshEditorMode::OnViewportInteractionInputAction( FEditorViewportClient& V
 							}
 						}
 					}
-					else if( SelectedMeshElementType == EEditableMeshElementType::Edge && EquippedEdgeAction == EMeshEditAction::InsertEdgeLoop )
-					{
-						const bool bActionNeedsHoverLocation = true;
-						StartAction( EMeshEditAction::InsertEdgeLoop, ViewportInteractor, bActionNeedsHoverLocation, LOCTEXT( "UndoInsertEdgeLoop", "Insert Edge Loop" ) );
-
-						bOutIsInputCaptured = true;
-					}
 					else if( SelectedMeshElementType == EEditableMeshElementType::Edge && EquippedEdgeAction == EMeshEditAction::ExtendEdge )
 					{
 						bWantToStartMoving = true;
@@ -5459,6 +5274,30 @@ void FMeshEditorMode::OnViewportInteractionInputAction( FEditorViewportClient& V
 
 						bOutIsInputCaptured = true;
 					}
+					else
+					{
+						for( TObjectIterator<UMeshEditorCommand> CommandCDOIter( RF_NoFlags ); CommandCDOIter; ++CommandCDOIter )
+						{
+							UMeshEditorCommand* CommandCDO = *CommandCDOIter;
+							if( !( CommandCDO->GetClass()->GetClassFlags() & CLASS_Abstract ) )
+							{
+								if( ( CommandCDO->GetElementType() == EEditableMeshElementType::Vertex && SelectedMeshElementType == EEditableMeshElementType::Vertex && EquippedVertexAction == CommandCDO->GetCommandName() ) ||
+									( CommandCDO->GetElementType() == EEditableMeshElementType::Edge && SelectedMeshElementType == EEditableMeshElementType::Edge  && EquippedEdgeAction == CommandCDO->GetCommandName() ) ||
+									( CommandCDO->GetElementType() == EEditableMeshElementType::Polygon && SelectedMeshElementType == EEditableMeshElementType::Polygon && EquippedPolygonAction == CommandCDO->GetCommandName() ) )
+								{
+									StartAction( EquippedEdgeAction, ViewportInteractor, CommandCDO->NeedsHoverLocation(), CommandCDO->GetUndoText() );
+
+									if( !CommandCDO->NeedsDraggingInitiated() )
+									{
+										bOutIsInputCaptured = true;
+									}
+
+									// Should always only be one candidate
+									break;
+								}
+							}
+						}
+					}
 				}
 				else 
 				{
@@ -5469,7 +5308,7 @@ void FMeshEditorMode::OnViewportInteractionInputAction( FEditorViewportClient& V
 
 						FSelectOrDeselectMeshElementsChangeInput ChangeInput;
 						ChangeInput.MeshElementsToDeselect.Add( SelectedMeshElements[ AlreadySelectedMeshElement ] );
-						UndoHelpers::StoreUndo( MeshEditorModeProxyObject, FSelectOrDeselectMeshElementsChange( MoveTemp( ChangeInput ) ).Execute( MeshEditorModeProxyObject ) );
+						TrackUndo( MeshEditorModeProxyObject, FSelectOrDeselectMeshElementsChange( MoveTemp( ChangeInput ) ).Execute( MeshEditorModeProxyObject ) );
 					}
 					else if( MeshElementSelectionMode == EEditableMeshElementType::Any || MeshElementSelectionMode == MeshEditorInteractorData.HoveredMeshElement.ElementAddress.ElementType )
 					{
@@ -5504,7 +5343,7 @@ void FMeshEditorMode::OnViewportInteractionInputAction( FEditorViewportClient& V
 						{
 							// If select by painting is disabled, add a transaction immediately
 							const FScopedTransaction Transaction( LOCTEXT( "SelectElement", "Select Element" ) );
-							UndoHelpers::StoreUndo( MeshEditorModeProxyObject, MoveTemp( RevertChange ) );
+							TrackUndo( MeshEditorModeProxyObject, MoveTemp( RevertChange ) );
 						}
 					}
 				}
@@ -5553,7 +5392,7 @@ void FMeshEditorMode::OnViewportInteractionInputAction( FEditorViewportClient& V
 							// canceled our transaction while the mouse was down.
 							if( GUndo != nullptr )
 							{
-								UndoHelpers::StoreUndo( MeshEditorModeProxyObject, MakeUnique<FCompoundChange>( MoveTemp( *SelectingByPaintingRevertChangeInput.Release() ) ) );
+								TrackUndo( MeshEditorModeProxyObject, MakeUnique<FCompoundChange>( MoveTemp( *SelectingByPaintingRevertChangeInput.Release() ) ) );
 							}
 						}
 						SelectingByPaintingRevertChangeInput.Reset();
@@ -5570,7 +5409,7 @@ void FMeshEditorMode::OnViewportInteractionInputAction( FEditorViewportClient& V
 }
 
 
-void FMeshEditorMode::StartAction( EMeshEditAction NewAction, UViewportInteractor* ActionInteractor, const bool bActionNeedsHoverLocation, const FText& UndoText )
+void FMeshEditorMode::StartAction( EMeshEditAction::Type NewAction, UViewportInteractor* ActionInteractor, const bool bActionNeedsHoverLocation, const FText& UndoText )
 {
 	// Don't start a new action without finishing the previous one!
 	check( ActiveAction == EMeshEditAction::None );
@@ -5599,7 +5438,7 @@ void FMeshEditorMode::FinishAction()
 	// @todo mesheditor: Make sure this is called before Undo is invoked (PreEditUndo!), otherwise the previous action will be undone instead of the active one
 
 	check( ActiveAction != EMeshEditAction::None );
-	check( !UndoHelpers::IsUndoSystemAvailable() || GEditor->IsTransactionActive() );	// Someone must have started a transaction! (It might not have been us though.)
+	check( GUndo == nullptr || GEditor->IsTransactionActive() );	// Someone must have started a transaction! (It might not have been us though.)
 
 	const bool bIsActionFinishing = true;
 
@@ -5612,9 +5451,9 @@ void FMeshEditorMode::FinishAction()
 	{
 		// @todo mesheditor: Drawing vertices will likely need to be a different kind of active action as it works differently to the others.
 		// For now, this just forces vertex drawing to be a "one shot" kind of mode.
-		EquippedVertexAction = EMeshEditAction::Move;
-		EquippedEdgeAction = EMeshEditAction::Move;
-		EquippedPolygonAction = EMeshEditAction::Move;
+		SetEquippedAction(EEditableMeshElementType::Vertex, EMeshEditAction::Move );
+		SetEquippedAction(EEditableMeshElementType::Edge, EMeshEditAction::Move );
+		SetEquippedAction( EEditableMeshElementType::Polygon, EMeshEditAction::Move );
 	}
 
 	if( bIsActionFinishing )
@@ -5980,7 +5819,7 @@ void FMeshEditorMode::PerformMarqueeSelect( EEditableMeshElementType ElementType
 			break;
 	}
 
-	UndoHelpers::StoreUndo( MeshEditorModeProxyObject, FSelectOrDeselectMeshElementsChange( MoveTemp( ChangeInput ) ).Execute( MeshEditorModeProxyObject ) );
+	TrackUndo( MeshEditorModeProxyObject, FSelectOrDeselectMeshElementsChange( MoveTemp( ChangeInput ) ).Execute( MeshEditorModeProxyObject ) );
 }
 
 
@@ -6093,20 +5932,20 @@ void FMeshEditorMode::OnActorSelectionChanged( const TArray<UObject*>& NewSelect
 		// be part of the undo history and we don't need to do anything.
 		if( GEditor->IsTransactionActive() )
 		{
-			FSelectOrDeselectMeshElementsChangeInput ChangeInput;
-
+			bool bAnyInvalidElementsSelected = false;
 			for( int32 SelectedElementIndex = 0; SelectedElementIndex < this->SelectedMeshElements.Num(); ++SelectedElementIndex )
 			{
 				const FMeshElement& SelectedMeshElement = this->SelectedMeshElements[ SelectedElementIndex ];
 				if( !SelectedMeshElement.Component.IsValid() || ( SelectedMeshElement.IsValidMeshElement() && !SelectedMeshElement.Component->GetOwner()->IsSelected() ) )
 				{
-					ChangeInput.MeshElementsToDeselect.Add( SelectedMeshElement );
+					bAnyInvalidElementsSelected = true;
+					break;
 				}
 			}
 
-			if( ChangeInput.MeshElementsToDeselect.Num() > 0 || ChangeInput.MeshElementsToSelect.Num() > 0 )
+			if( bAnyInvalidElementsSelected )
 			{
-				UndoHelpers::StoreUndo( MeshEditorModeProxyObject, FDeselectAllMeshElementsChange( MoveTemp( FDeselectAllMeshElementsChangeInput() ) ).Execute( MeshEditorModeProxyObject ) );
+				DeselectAllMeshElements();
 			}
 		}
 	}
@@ -6158,7 +5997,7 @@ void FMeshEditorMode::MakeVRRadialMenuActionsMenu(FMenuBuilder& MenuBuilder, TSh
 			FSlateIcon(FMeshEditorStyle::GetStyleSetName(), "MeshEditorMode.PolyMove"),
 			FUIAction
 			(
-				FExecuteAction::CreateLambda([this] { EquippedPolygonAction = EMeshEditAction::Move; }),
+				FExecuteAction::CreateLambda([this] { SetEquippedAction(EEditableMeshElementType::Polygon, EMeshEditAction::Move ); }),
 				FCanExecuteAction::CreateSP(this, &FMeshEditorMode::IsMeshElementTypeSelectedOrIsActiveSelectionMode, EEditableMeshElementType::Polygon),
 				FIsActionChecked::CreateLambda([this] { return (EquippedPolygonAction == EMeshEditAction::Move); })
 				),
@@ -6183,7 +6022,7 @@ void FMeshEditorMode::MakeVRRadialMenuActionsMenu(FMenuBuilder& MenuBuilder, TSh
 			FSlateIcon(FMeshEditorStyle::GetStyleSetName(), "MeshEditorMode.PolyExtrude"),
 			FUIAction
 			(
-				FExecuteAction::CreateLambda([this] { EquippedPolygonAction = EMeshEditAction::ExtrudePolygon; }),
+				FExecuteAction::CreateLambda([this] { SetEquippedAction( EEditableMeshElementType::Polygon, EMeshEditAction::ExtrudePolygon ); }),
 				FCanExecuteAction::CreateSP(this, &FMeshEditorMode::IsMeshElementTypeSelectedOrIsActiveSelectionMode, EEditableMeshElementType::Polygon),
 				FIsActionChecked::CreateLambda([this] { return (EquippedPolygonAction == EMeshEditAction::ExtrudePolygon); })
 			),
@@ -6196,7 +6035,7 @@ void FMeshEditorMode::MakeVRRadialMenuActionsMenu(FMenuBuilder& MenuBuilder, TSh
 			FSlateIcon(FMeshEditorStyle::GetStyleSetName(), "MeshEditorMode.PolyInset"),
 			FUIAction
 			(
-				FExecuteAction::CreateLambda([this] { EquippedPolygonAction = EMeshEditAction::InsetPolygon; }),
+				FExecuteAction::CreateLambda([this] { SetEquippedAction( EEditableMeshElementType::Polygon, EMeshEditAction::InsetPolygon ); }),
 				FCanExecuteAction::CreateSP(this, &FMeshEditorMode::IsMeshElementTypeSelectedOrIsActiveSelectionMode, EEditableMeshElementType::Polygon),
 				FIsActionChecked::CreateLambda([this] { return (EquippedPolygonAction == EMeshEditAction::InsetPolygon); })
 			),
@@ -6212,7 +6051,7 @@ void FMeshEditorMode::MakeVRRadialMenuActionsMenu(FMenuBuilder& MenuBuilder, TSh
 			FSlateIcon(FMeshEditorStyle::GetStyleSetName(), "MeshEditorMode.EdgeMove"),
 			FUIAction
 			(
-				FExecuteAction::CreateLambda([this] { EquippedEdgeAction = EMeshEditAction::Move; }),
+				FExecuteAction::CreateLambda([this] { SetEquippedAction( EEditableMeshElementType::Edge, EMeshEditAction::Move ); }),
 				FCanExecuteAction::CreateSP(this, &FMeshEditorMode::IsMeshElementTypeSelectedOrIsActiveSelectionMode, EEditableMeshElementType::Edge),
 				FIsActionChecked::CreateLambda([this] { return (EquippedEdgeAction == EMeshEditAction::Move); })
 			),
@@ -6255,26 +6094,13 @@ void FMeshEditorMode::MakeVRRadialMenuActionsMenu(FMenuBuilder& MenuBuilder, TSh
 			NAME_None,
 			EUserInterfaceActionType::CollapsedButton
 			);
-		MenuBuilder.AddMenuEntry(
-			LOCTEXT("InsertEdgeLoop", "Insert Loop"),
-			FText(),
-			FSlateIcon(FMeshEditorStyle::GetStyleSetName(), "MeshEditorMode.EdgeInsert"),
-			FUIAction
-			(
-				FExecuteAction::CreateLambda([this] { EquippedEdgeAction = EMeshEditAction::InsertEdgeLoop; }),
-				FCanExecuteAction::CreateSP(this, &FMeshEditorMode::IsMeshElementTypeSelectedOrIsActiveSelectionMode, EEditableMeshElementType::Edge),
-				FIsActionChecked::CreateLambda([this] { return (EquippedEdgeAction == EMeshEditAction::InsertEdgeLoop); })
-				),
-			NAME_None,
-			EUserInterfaceActionType::ToggleButton
-			);
 		MenuBuilder.AddMenuEntry(
 			LOCTEXT("ExtendEdge", "Extend"),
 			FText(),
 			FSlateIcon(FMeshEditorStyle::GetStyleSetName(), "MeshEditorMode.EdgeExtend"),
 			FUIAction
 			(
-				FExecuteAction::CreateLambda([this] { EquippedEdgeAction = EMeshEditAction::ExtendEdge; }),
+				FExecuteAction::CreateLambda([this] { SetEquippedAction(EEditableMeshElementType::Edge, EMeshEditAction::ExtendEdge ); }),
 				FCanExecuteAction::CreateSP(this, &FMeshEditorMode::IsMeshElementTypeSelectedOrIsActiveSelectionMode, EEditableMeshElementType::Edge),
 				FIsActionChecked::CreateLambda([this] { return (EquippedEdgeAction == EMeshEditAction::ExtendEdge); })
 				),
@@ -6290,7 +6116,7 @@ void FMeshEditorMode::MakeVRRadialMenuActionsMenu(FMenuBuilder& MenuBuilder, TSh
 			FSlateIcon(FMeshEditorStyle::GetStyleSetName(), "MeshEditorMode.VertexMove"),
 			FUIAction
 			(
-				FExecuteAction::CreateLambda([this] { EquippedVertexAction = EMeshEditAction::Move; }),
+				FExecuteAction::CreateLambda([this] { SetEquippedAction( EEditableMeshElementType::Vertex, EMeshEditAction::Move ); }),
 				FCanExecuteAction::CreateSP(this, &FMeshEditorMode::IsMeshElementTypeSelectedOrIsActiveSelectionMode, EEditableMeshElementType::Vertex),
 				FIsActionChecked::CreateLambda([this] { return (EquippedVertexAction == EMeshEditAction::Move); })
 				),
@@ -6327,7 +6153,7 @@ void FMeshEditorMode::MakeVRRadialMenuActionsMenu(FMenuBuilder& MenuBuilder, TSh
 			FSlateIcon(FMeshEditorStyle::GetStyleSetName(), "MeshEditorMode.VertexExtend"),
 			FUIAction
 			(
-				FExecuteAction::CreateLambda([this] { EquippedVertexAction = EMeshEditAction::ExtendVertex; }),
+				FExecuteAction::CreateLambda([this] { SetEquippedAction( EEditableMeshElementType::Vertex, EMeshEditAction::ExtendVertex ); }),
 				FCanExecuteAction::CreateSP(this, &FMeshEditorMode::IsMeshElementTypeSelectedOrIsActiveSelectionMode, EEditableMeshElementType::Vertex),
 				FIsActionChecked::CreateLambda([this] { return (EquippedVertexAction == EMeshEditAction::ExtendVertex); })
 				),
@@ -6347,6 +6173,103 @@ void FMeshEditorMode::MakeVRRadialMenuActionsMenu(FMenuBuilder& MenuBuilder, TSh
 			EUserInterfaceActionType::CollapsedButton
 			);
 	}
+
+	for( TObjectIterator<UMeshEditorCommand> CommandCDOIter( RF_NoFlags ); CommandCDOIter; ++CommandCDOIter )
+	{
+		UMeshEditorCommand* CommandCDO = *CommandCDOIter;
+		if( !( CommandCDO->GetClass()->GetClassFlags() & CLASS_Abstract ) )
+		{
+			CommandCDO->AddToVRRadialMenuActionsMenu( *this, MenuBuilder, CommandList, FMeshEditorStyle::GetStyleSetName(), VRMode );
+		}
+	}
+}
+
+
+EMeshEditAction::Type FMeshEditorMode::GetEquippedAction( const EEditableMeshElementType ForElementType ) const
+{
+	EMeshEditAction::Type EquippedAction = EMeshEditAction::None;
+
+	switch( ForElementType )
+	{
+		case EEditableMeshElementType::Vertex:
+			EquippedAction = EquippedVertexAction;
+			break;
+
+		case EEditableMeshElementType::Edge:
+			EquippedAction = EquippedEdgeAction;
+			break;
+
+		case EEditableMeshElementType::Polygon:
+			EquippedAction = EquippedEdgeAction;
+			break;
+
+		default:
+			check( 0 );
+	}
+
+	return EquippedAction;
+}
+
+
+void FMeshEditorMode::SetEquippedAction( const EEditableMeshElementType ForElementType, const EMeshEditAction::Type ActionToEquip )
+{
+	switch( ForElementType )
+	{
+		case EEditableMeshElementType::Vertex:
+			EquippedVertexAction = ActionToEquip;
+			break;
+
+		case EEditableMeshElementType::Edge:
+			EquippedEdgeAction = ActionToEquip;
+			break;
+
+		case EEditableMeshElementType::Polygon:
+			EquippedEdgeAction = ActionToEquip;
+			break;
+
+		default:
+			check( 0 );
+	}
 }
 
+
+void FMeshEditorMode::TrackUndo( UObject* Object, TUniquePtr<FChange> RevertChange )
+{
+	if( RevertChange.IsValid() )
+	{
+		check( Object != nullptr );
+
+		if( !bIsCapturingUndoForPreview )
+		{
+			// If we're finalizing the action, this will save the undo state for everything that happened in this function, including
+			// selection changes.  
+
+			// Did you forget to use an FScopedTransaction?  If GUndo was null, then most likely we forgot to wrap this call within an editor transaction.
+			// The only exception is in Simulate mode, where Undo is not allowed.
+			check( GUndo != nullptr || GEditor == nullptr || GEditor->bIsSimulatingInEditor );
+
+			GUndo->StoreUndo( Object, MoveTemp( RevertChange ) );
+		}
+		else
+		{
+			// Otherwise, we'll store the commands to undo in our 'PreviewRevertChanges' member, so they can be
+			// rolled back at the beginning of the next frame before any new interactions take place.  This allows the user to preview
+			// (potentially highly destructive) changes live!
+
+			// If the object is a mesh, make sure we've started to modify it
+			UEditableMesh* EditableMesh = Cast<UEditableMesh>( Object );
+			if( EditableMesh != nullptr )
+			{
+				// StartModification() must have already been called, otherwise it's too late at this point -- the mesh has been changed
+				// while the render thread was still using it.  Bad things will happen.  So we assert here.
+				check( ActiveActionModifiedMeshes.Contains( EditableMesh ) );
+			}
+
+			// NOTE: These changes will be rolled back in the opposite order they were added to the list
+			PreviewRevertChanges.Add( MakeTuple( Object, MoveTemp( RevertChange ) ) );
+		}
+	}
+}
+
+
 #undef LOCTEXT_NAMESPACE
diff --git a/Engine/Plugins/Editor/MeshEditor/Source/MeshEditor/MeshEditorMode.h b/Engine/Plugins/Editor/MeshEditor/Source/MeshEditor/MeshEditorMode.h
--- a/Engine/Plugins/Editor/MeshEditor/Source/MeshEditor/MeshEditorMode.h
+++ b/Engine/Plugins/Editor/MeshEditor/Source/MeshEditor/MeshEditorMode.h
@@ -30,7 +30,7 @@ public:
 /**
  * Mesh Editor Mode.  Extends editor viewports with the ability to edit meshes
  */
-class FMeshEditorMode : public FEdMode, protected IMeshEditorModeUIContract, protected IMeshEditorModeEditingContract
+class FMeshEditorMode : public FEdMode, protected IMeshEditorModeUIContract
 {
 public:
 
@@ -154,6 +154,8 @@ protected:
 	}
 	virtual void SetMeshElementSelectionMode( EEditableMeshElementType ElementType ) override;
 	virtual EEditableMeshElementType GetSelectedMeshElementType() const override;
+	virtual bool IsMeshElementTypeSelected( EEditableMeshElementType ElementType ) const override { return GetSelectedMeshElementType() == ElementType; }
+	virtual bool IsMeshElementTypeSelectedOrIsActiveSelectionMode( EEditableMeshElementType ElementType ) const override { return GetMeshElementSelectionMode() == ElementType || GetSelectedMeshElementType() == ElementType; }
 	virtual const TArray<TTuple<TSharedPtr<FUICommandInfo>, FUIAction>>& GetVertexActions() const override { return VertexActions; }
 	virtual const TArray<TTuple<TSharedPtr<FUICommandInfo>, FUIAction>>& GetEdgeActions() const override { return EdgeActions; }
 	virtual const TArray<TTuple<TSharedPtr<FUICommandInfo>, FUIAction>>& GetPolygonActions() const override { return PolygonActions; }
@@ -161,38 +163,41 @@ protected:
 	virtual void SetEditingPerInstance( bool bPerInstance ) override { bPerInstanceEdits = bPerInstance; }
 	virtual void PropagateInstanceChanges() override;
 	virtual bool CanPropagateInstanceChanges() const override;
+	virtual EMeshEditAction::Type GetEquippedAction( const EEditableMeshElementType ForElementType ) const override;
+	virtual void SetEquippedAction( const EEditableMeshElementType ForElementType, const EMeshEditAction::Type ActionToEquip ) override;
 
 	// IMeshEditorModeEditingContract interface
-	virtual EMeshEditAction GetActiveAction() const	override { return ActiveAction; }
+	virtual EMeshEditAction::Type GetActiveAction() const override { return ActiveAction; }
+	virtual void TrackUndo( UObject* Object, TUniquePtr<FChange> RevertChange ) override;
 	virtual void CommitSelectedMeshes() override;
 	virtual void GetSelectedMeshesAndVertices( TMap<UEditableMesh*, TArray<FMeshElement>>& OutMeshesAndVertices ) override 
 	{
-		GetSelectedMeshesAndElements( EEditableMeshElementType::Vertex, OutMeshesAndVertices ); 
+		GetSelectedMeshesAndElements( EEditableMeshElementType::Vertex, /* Out */ OutMeshesAndVertices ); 
 	}
 	virtual void GetSelectedMeshesAndEdges( TMap<UEditableMesh*, TArray<FMeshElement>>& OutMeshesAndEdges ) override 
 	{
-		GetSelectedMeshesAndElements( EEditableMeshElementType::Edge, OutMeshesAndEdges ); 
+		GetSelectedMeshesAndElements( EEditableMeshElementType::Edge, /* Out */ OutMeshesAndEdges );
 	}
 	virtual void GetSelectedMeshesAndPolygons( TMap<UEditableMesh*, TArray<FMeshElement>>& OutMeshesAndPolygons ) override 
 	{
-		GetSelectedMeshesAndElements( EEditableMeshElementType::Polygon, OutMeshesAndPolygons ); 
+		GetSelectedMeshesAndElements( EEditableMeshElementType::Polygon, /* Out */ OutMeshesAndPolygons );
 	}
 	virtual void SelectMeshElements( const TArray<FMeshElement>& MeshElementsToSelect ) override;
+	virtual void DeselectAllMeshElements() override;
 	virtual void DeselectMeshElements( const TArray<FMeshElement>& MeshElementsToDeselect ) override;
 	virtual void DeselectMeshElements( const TMap<UEditableMesh*, TArray<FMeshElement>>& MeshElementsToDeselect ) override;
 	virtual void ClearInvalidSelectedElements() override;
+	virtual void FindEdgeSplitUnderInteractor( UViewportInteractor* ViewportInteractor, const UEditableMesh* EditableMesh, const TArray<FMeshElement>& EdgeElements, TArray<float>& OutSplits ) override;
+	virtual class UViewportInteractor* GetActiveActionInteractor() override
+	{
+		return ActiveActionInteractor;
+	}
 
 
 
 	/** Gets the container of all the assets used in the mesh editor */
 	const class UMeshEditorAssetContainer& GetAssetContainer() const;
 
-	/** Returns whether the specified element type is selected */
-	bool IsMeshElementTypeSelected( EEditableMeshElementType ElementType ) const { return GetSelectedMeshElementType() == ElementType; }
-
-	/** Returns whether either the specified element type is selected, or we're in the selection mode for that element type */
-	bool IsMeshElementTypeSelectedOrIsActiveSelectionMode( EEditableMeshElementType ElementType ) const { return GetMeshElementSelectionMode() == ElementType || GetSelectedMeshElementType() == ElementType; }
-
 	/** Fills the specified dynamic mesh builder with primitives to render a mesh vertex */
 	void AddVertexToDynamicMesh( const UEditableMesh& EditableMesh, const FTransform& CameraToWorld, const FMatrix& ComponentToWorldMatrix, const FVertexID VertexID, const FColor ColorAndOpacity, const float SizeBias, const bool bApplyDepthBias, class FDynamicMeshBuilder& MeshBuilder );
 		
@@ -248,7 +253,7 @@ protected:
 	void UpdateCameraToWorldTransform( const FEditorViewportClient& ViewportClient );
 
 	/** Begins an action */
-	void StartAction( const EMeshEditAction NewAction, class UViewportInteractor* ActionInteractor, const bool bActionNeedsHoverLocation, const FText& UndoText );
+	void StartAction( const EMeshEditAction::Type NewAction, class UViewportInteractor* ActionInteractor, const bool bActionNeedsHoverLocation, const FText& UndoText );
 
 	/** Ends an action that's currently in progress.  Usually called when the user commits a change by clicking/releasing, but can
 	    also be called when the user begins a new action while inertia is still influencing the active action */
@@ -257,10 +262,10 @@ protected:
 	/** Binds UI commands to actions for the mesh editor */
 	void BindCommands();
 
-	void RegisterCommonEditingMode( const TSharedPtr<FUICommandInfo>& Command, EMeshEditAction EditingMode );
-	void RegisterVertexEditingMode( const TSharedPtr<FUICommandInfo>& Command, EMeshEditAction EditingMode );
-	void RegisterEdgeEditingMode( const TSharedPtr<FUICommandInfo>& Command, EMeshEditAction EditingMode );
-	void RegisterPolygonEditingMode( const TSharedPtr<FUICommandInfo>& Command, EMeshEditAction EditingMode );
+	void RegisterCommonEditingMode( const TSharedPtr<FUICommandInfo>& Command, EMeshEditAction::Type EditingMode );
+	void RegisterVertexEditingMode( const TSharedPtr<FUICommandInfo>& Command, EMeshEditAction::Type EditingMode );
+	void RegisterEdgeEditingMode( const TSharedPtr<FUICommandInfo>& Command, EMeshEditAction::Type EditingMode );
+	void RegisterPolygonEditingMode( const TSharedPtr<FUICommandInfo>& Command, EMeshEditAction::Type EditingMode );
 
 	void RegisterCommand( const TSharedPtr<FUICommandInfo>& Command, const FExecuteAction& ExecuteAction );
 	void RegisterVertexCommand( const TSharedPtr<FUICommandInfo>& Command, const FExecuteAction& ExecuteAction );
@@ -324,18 +329,6 @@ protected:
 	/** Helper function that returns a map keying an editable mesh with its selected elements */
 	void GetSelectedMeshesAndElements( EEditableMeshElementType ElementType, TMap<UEditableMesh*, TArray<FMeshElement>>& OutMeshesAndElements );
 
-	/** Helper function that converts selected vertices into a map of meshes to their respective selected vertex elements and vertex IDs */
-	void GetSelectedMeshesAndVertices( TMap< UEditableMesh*, TArray< TTuple< FMeshElement, FVertexID > > >& OutMeshesAndVertices );
-
-	/** Helper function that converts selected edges into a map of meshes to their respective selected edge elements and edge IDs */
-	void GetSelectedMeshesAndEdges( TMap< UEditableMesh*, TArray< TTuple< FMeshElement, FEdgeID > > >& OutMeshesAndEdges );
-
-	/** Helper function that converts selected polygons into a map of meshes to their respective selected polygon elements and polygon refs */
-	void GetSelectedMeshesAndPolygons( TMap< UEditableMesh*, TArray< TTuple< FMeshElement, FPolygonRef > > >& OutMeshesAndPolygons );
-
-	/** Given an interactor and a mesh, finds edges under the interactor along with their exact split position (progress along the edge) */
-	void FindEdgeSplitUnderInteractor( UViewportInteractor* ViewportInteractor, const UEditableMesh* EditableMesh, const TArray<TTuple<FMeshElement, FEdgeID>>& Edges, TArray<float>& OutSplits );
-
 	/** Selects elements of the given type captured by the last marquee select */
 	void PerformMarqueeSelect( EEditableMeshElementType ElementType );
 
@@ -361,10 +354,11 @@ protected:
 	bool IsActive() const { return ( ViewportWorldInteraction != nullptr ); }
 
 	/** Plays sound when starting a mesh edit action */
-	void PlayStartActionSound(EMeshEditAction NewAction, UViewportInteractor* ActionInteractor = nullptr);
+	void PlayStartActionSound(EMeshEditAction::Type NewAction, UViewportInteractor* ActionInteractor = nullptr);
 
 	/** Plays sound when mesh edit action was finished */
-	void PlayFinishActionSound(EMeshEditAction NewAction, UViewportInteractor* ActionInteractor = nullptr);
+	void PlayFinishActionSound(EMeshEditAction::Type NewAction, UViewportInteractor* ActionInteractor = nullptr);
+
 
 protected:
 
@@ -530,17 +524,34 @@ protected:
 	FTrackingTransaction TrackingTransaction;
 
 	/** The next action that will be started when interacting with a selected vertex */
-	EMeshEditAction EquippedVertexAction;
+	EMeshEditAction::Type EquippedVertexAction;
 
 	/** The next action that will be started when interacting with a selected edge */
-	EMeshEditAction EquippedEdgeAction;
+	EMeshEditAction::Type EquippedEdgeAction;
 
 	/** The next action that will be started when interacting with a selected polygon */
-	EMeshEditAction EquippedPolygonAction;
+	EMeshEditAction::Type EquippedPolygonAction;
 
 	/** The interactive action currently being performed (and previewed).  These usually happen over multiple frames, and
 	    result in a 'final' application of the change that performs a more exhaustive (and more expensive) update. */
-	EMeshEditAction ActiveAction;
+	EMeshEditAction::Type ActiveAction;
+
+	/** Whether we're actually in the middile of updating the active action.  This means that StoreUndo() will behave
+	    differently in this case -- instead of pushing undo data to the editor, we'll capture it temporarily in PreviewRevertChanges,
+		so that we can roll it back at the beginning of the next frame. */
+	bool bIsCapturingUndoForPreview;
+
+	/** When interactively dragging to preview a change (that might not be fully committed), this is the Change that
+	    will be used to roll back the previewed alternative from the previous frame */
+	TArray<TTuple<class UObject*,TUniquePtr<FChange>>> PreviewRevertChanges;
+
+	/** The set of meshes that we are expecting to be modified while updating the current active action.  These meshes
+	    will have StartModification() called at the beginning of UpdateActiveAction() and EndModification() afterwards. */
+	TSet<class UEditableMesh*> ActiveActionModifiedMeshes;
+
+	/** Proxy UObject to pass to the undo system when performing interactions that affect state of the mode itself,
+	    such as the selection set.  We need this because the UE4 undo system requires a UObject, but we're an FEdMode. */
+	UMeshEditorModeProxyObject* MeshEditorModeProxyObject;
 
 	/** When performing an interactive action that was initiated using an interactor, this is the interactor that was used. */
 	class UViewportInteractor* ActiveActionInteractor;
@@ -591,7 +602,7 @@ protected:
 	//
 
 	/** When splitting an edge and dragging a vertex, this is the list of edges that will be split */
-	TMap< UEditableMesh*, TArray< TTuple< FMeshElement, FEdgeID > > > SplitEdgeMeshesAndEdgesToSplit;
+	TMap< UEditableMesh*, TArray< FMeshElement > > SplitEdgeMeshesAndEdgesToSplit;
 
 	/** When splitting an edge and dragging a vertex, this is the list of split positions along those edges */
 	TArray<float> SplitEdgeSplitList;
@@ -615,17 +626,6 @@ protected:
 
 	// ...
 
-	/** When interactively dragging to preview a change (that might not be fully committed), this is the Change that
-	    will be used to roll back the previewed alternative from the previous frame */
-	TArray<TTuple<class UObject*,TUniquePtr<FChange>>> PreviewRevertChanges;
-
-	/** Whether topology changed while we were applying the preview changes that we might revert later */
-	EMeshTopologyChange PreviewTopologyChange;
-
-	/** Proxy UObject to pass to the undo system when performing interactions that affect state of the mode itself, 
-	    such as the selection set.  We need this because the UE4 undo system requires a UObject, but we're an FEdMode. */
-	UMeshEditorModeProxyObject* MeshEditorModeProxyObject;
-
 	/** When selecting by painting (EMeshEditAction::SelectByPainting), this is the compound change that can be applied to
 	    roll back the change to select.  We'll build this up as the user is painting select, then store it in the undo buffer */
 	TUniquePtr<FCompoundChangeInput> SelectingByPaintingRevertChangeInput;
diff --git a/Engine/Plugins/Editor/MeshEditor/Source/MeshEditor/Public/IMeshEditorModeEditingContract.h b/Engine/Plugins/Editor/MeshEditor/Source/MeshEditor/Public/IMeshEditorModeEditingContract.h
--- a/Engine/Plugins/Editor/MeshEditor/Source/MeshEditor/Public/IMeshEditorModeEditingContract.h
+++ b/Engine/Plugins/Editor/MeshEditor/Source/MeshEditor/Public/IMeshEditorModeEditingContract.h
@@ -6,66 +6,72 @@
 #include "MeshElement.h"
 
 
-class IMeshEditorModeEditingContract
+// @todo mesheditor extensibility: This should probably be removed after we've evicted all current mesh editing actions to another module
+namespace EMeshEditAction
 {
+	typedef FName Type;
 
-public:
+	/** Nothing going on right now */
+	MESHEDITOR_API extern const FName None;
+
+	/** Selecting mesh elements by 'painting' over multiple elements */
+	MESHEDITOR_API extern const FName SelectByPainting;
 
-	enum class EMeshEditAction
-	{
-		/** Nothing going on right now */
-		None,
+	/** Moving elements using a transform gizmo */
+	MESHEDITOR_API extern const FName MoveUsingGizmo;
 
-		/** Selecting mesh elements by 'painting' over multiple elements */
-		SelectByPainting,
+	/** Moving selected mesh elements (vertices, edges or polygons) */
+	MESHEDITOR_API extern const FName Move;
 
-		/** Moving elements using a transform gizmo */
-		MoveUsingGizmo,
+	/** Split an edge by inserting a vertex.  You can drag to preview where the vertex will be inserted. */
+	MESHEDITOR_API extern const FName SplitEdge;
 
-		/** Moving selected mesh elements (vertices, edges or polygons) */
-		Move,
+	/** Splits an edge by inserting a new vertex, then immediately starts dragging that vertex */
+	MESHEDITOR_API extern const FName SplitEdgeAndDragVertex;
 
-		/** Split an edge by inserting a vertex.  You can drag to preview where the vertex will be inserted. */
-		SplitEdge,
+	/** Extrude polygon by making a copy of it and allowing you to shift it along the polygon normal axis */
+	MESHEDITOR_API extern const FName ExtrudePolygon;
 
-		/** Splits an edge by inserting a new vertex, then immediately starts dragging that vertex */
-		SplitEdgeAndDragVertex,
+	/** Extrude polygon by making a copy of it and allowing you to move it around freely */
+	MESHEDITOR_API extern const FName FreelyExtrudePolygon;
 
-		/** Insert an edge loop */
-		InsertEdgeLoop,
+	/** Inset polygon by replacing it with a new polygon that is bordered by polygons of a specific relative size */
+	MESHEDITOR_API extern const FName InsetPolygon;
 
-		/** Extrude polygon by making a copy of it and allowing you to shift it along the polygon normal axis */
-		ExtrudePolygon,
+	/** Bevel polygons by adding angled bordering polygons of a specific relative size */
+	MESHEDITOR_API extern const FName BevelPolygon;
 
-		/** Extrude polygon by making a copy of it and allowing you to move it around freely */
-		FreelyExtrudePolygon,
+	/** Extend an edge by making a copy of it and allowing you to move it around */
+	MESHEDITOR_API extern const FName ExtendEdge;
 
-		/** Inset polygon by replacing it with a new polygon that is bordered by polygons of a specific relative size */
-		InsetPolygon,
+	/** Extend a vertex by making a copy of it, creating new polygons to join the geometry together */
+	MESHEDITOR_API extern const FName ExtendVertex;
 
-		/** Bevel polygons by adding angled bordering polygons of a specific relative size */
-		BevelPolygon,
+	/** For subdivision meshes, edits how sharp a vertex corner is by dragging in space */
+	MESHEDITOR_API extern const FName EditVertexCornerSharpness;
 
-		/** Extend an edge by making a copy of it and allowing you to move it around */
-		ExtendEdge,
+	/** For subdivision meshes, edits how sharp an edge crease is by dragging in space */
+	MESHEDITOR_API extern const FName EditEdgeCreaseSharpness;
 
-		/** Extend a vertex by making a copy of it, creating new polygons to join the geometry together */
-		ExtendVertex,
+	/** Freehand vertex drawing */
+	MESHEDITOR_API extern const FName DrawVertices;
+}
 
-		/** For subdivision meshes, edits how sharp a vertex corner is by dragging in space */
-		EditVertexCornerSharpness,
 
-		/** For subdivision meshes, edits how sharp an edge crease is by dragging in space */
-		EditEdgeCreaseSharpness,
 
-		/** Freehand vertex drawing */
-		DrawVertices,
-	};
+class IMeshEditorModeEditingContract
+{
+
+public:
 
 
 	/** Gets the interactive action currently being performed (and previewed).  These usually happen over multiple frames, and
 	    result in a 'final' application of the change that performs a more exhaustive (and more expensive) update. */
-	virtual EMeshEditAction GetActiveAction() const = 0;
+	virtual EMeshEditAction::Type GetActiveAction() const = 0;
+
+	/** Stores undo state for the specified object.  This will store the state different depending on whether we're
+	    currently in the middle of previewing a temporary change to meshes (bIsCapturingUndoForPreview) */
+	virtual void TrackUndo( UObject* Object, TUniquePtr<FChange> RevertChange ) = 0;
 
 	virtual void GetSelectedMeshesAndVertices( TMap<class UEditableMesh*, TArray<FMeshElement>>& OutMeshesAndVertices ) = 0;
 	virtual void GetSelectedMeshesAndEdges( TMap<class UEditableMesh*, TArray<FMeshElement>>& OutMeshesAndEdges ) = 0;
@@ -74,6 +80,9 @@ public:
 	/** Selects the specified mesh elements */
 	virtual void SelectMeshElements( const TArray<FMeshElement>& MeshElementsToSelect ) = 0;
 
+	/** Deselects all mesh elements */
+	virtual void DeselectAllMeshElements() = 0;
+
 	/** Deselects the specified mesh elements */
 	virtual void DeselectMeshElements( const TArray<FMeshElement>& MeshElementsToDeselect ) = 0;
 	virtual void DeselectMeshElements( const TMap<UEditableMesh*, TArray<FMeshElement>>& MeshElementsToDeselect ) = 0;
@@ -84,4 +93,10 @@ public:
 	/** Clears hover and selection on mesh elements that may no longer be valid.  You'll want to call this if you change the mesh topology */
 	virtual void ClearInvalidSelectedElements() = 0;
 
+	/** Given an interactor and a mesh, finds edges under the interactor along with their exact split position (progress along the edge) */
+	virtual void FindEdgeSplitUnderInteractor( class UViewportInteractor* ViewportInteractor, const UEditableMesh* EditableMesh, const TArray<FMeshElement>& EdgeElements, TArray<float>& OutSplits ) = 0;
+
+	/** When performing an interactive action that was initiated using an interactor, this is the interactor that was used. */
+	virtual class UViewportInteractor* GetActiveActionInteractor() = 0;
+
 };
diff --git a/Engine/Plugins/Editor/MeshEditor/Source/MeshEditor/Public/IMeshEditorModeUIContract.h b/Engine/Plugins/Editor/MeshEditor/Source/MeshEditor/Public/IMeshEditorModeUIContract.h
new file mode 100644
--- /dev/null
+++ b/Engine/Plugins/Editor/MeshEditor/Source/MeshEditor/Public/IMeshEditorModeUIContract.h
@@ -0,0 +1,47 @@
+// Copyright 1998-2015 Epic Games, Inc. All Rights Reserved.
+
+#pragma once
+
+#include "EditableMeshTypes.h"
+#include "IMeshEditorModeEditingContract.h"
+
+class IMeshEditorModeUIContract : public IMeshEditorModeEditingContract
+{
+
+public:
+
+	/** Returns the current selection mode we're in */
+	virtual EEditableMeshElementType GetMeshElementSelectionMode() const = 0;
+	
+	/** Sets the mesh element selection mode to use */
+	virtual void SetMeshElementSelectionMode( EEditableMeshElementType ElementType ) = 0;
+
+	/** Returns the type of elements that are selected right now, or Invalid if nothing is selected */
+	virtual EEditableMeshElementType GetSelectedMeshElementType() const = 0;
+
+	/** Returns whether the specified element type is selected */
+	virtual bool IsMeshElementTypeSelected( EEditableMeshElementType ElementType ) const = 0;
+
+	/** Returns whether either the specified element type is selected, or we're in the selection mode for that element type */
+	virtual bool IsMeshElementTypeSelectedOrIsActiveSelectionMode( EEditableMeshElementType ElementType ) const = 0;
+
+	virtual const TArray<TTuple<TSharedPtr<FUICommandInfo>, FUIAction>>& GetVertexActions() const = 0;
+	virtual const TArray<TTuple<TSharedPtr<FUICommandInfo>, FUIAction>>& GetEdgeActions() const = 0;
+	virtual const TArray<TTuple<TSharedPtr<FUICommandInfo>, FUIAction>>& GetPolygonActions() const = 0;
+
+	virtual bool IsEditingPerInstance() const = 0;
+	virtual void SetEditingPerInstance( bool bPerInstance ) = 0;
+
+	/** Propagates instance changes to the static mesh asset */
+	virtual void PropagateInstanceChanges() = 0;
+
+	/** Whether there are instance changes which can be propagated */
+	virtual bool CanPropagateInstanceChanges() const = 0;
+
+	/** Returns the current action to use when interacting the next time with the specified type of mesh element selection mode */
+	virtual EMeshEditAction::Type GetEquippedAction( const EEditableMeshElementType ForElementType ) const = 0;
+
+	/** Sets the current action to use when interacting the next time with the specified type of mesh element selection mode */
+	virtual void SetEquippedAction( const EEditableMeshElementType ForElementType, const EMeshEditAction::Type ActionToEquip ) = 0;
+
+};
diff --git a/Engine/Plugins/Editor/MeshEditor/Source/MeshEditor/Public/MeshEditorCommands.h b/Engine/Plugins/Editor/MeshEditor/Source/MeshEditor/Public/MeshEditorCommands.h
--- a/Engine/Plugins/Editor/MeshEditor/Source/MeshEditor/Public/MeshEditorCommands.h
+++ b/Engine/Plugins/Editor/MeshEditor/Source/MeshEditor/Public/MeshEditorCommands.h
@@ -2,6 +2,8 @@
 
 #pragma once
 
+#include "MeshElement.h"
+#include "UIAction.h"
 #include "MeshEditorCommands.generated.h"
 
 
@@ -12,11 +14,58 @@ class MESHEDITOR_API UMeshEditorCommand : public UObject
 
 public:
 
+	/** What type of mesh element is this command for? */
+	virtual EEditableMeshElementType GetElementType() const PURE_VIRTUAL(,return EEditableMeshElementType::Invalid;);
+
 	/** Registers the UI command for this mesh editor command */
 	virtual void RegisterUICommand( class FBindingContext* BindingContext ) PURE_VIRTUAL(,);
 
 	/** Runs this command */
-	virtual void Execute( class IMeshEditorModeEditingContract& MeshEditorMode, bool& bOutWasSuccessful ) PURE_VIRTUAL(,);
+	virtual void Execute( class IMeshEditorModeEditingContract& MeshEditorMode )
+	{
+	}
+
+	/** Applies this command every frame while dragging */
+	virtual void ApplyDuringDrag( class IMeshEditorModeEditingContract& MeshEditorMode, bool& bOutShouldDeselectAllFirst, TArray<FMeshElement>& OutMeshElementsToSelect )
+	{
+	}
+
+	/** Allows this command to directly add a button to the VR Mode's radial menu */
+	virtual void AddToVRRadialMenuActionsMenu( class IMeshEditorModeUIContract& MeshEditorMode, class FMenuBuilder& MenuBuilder, TSharedPtr<FUICommandList> CommandList, const FName TEMPHACK_StyleSetName, class UVREditorMode* VRMode )
+	{
+	}
+
+	/** Gets the name of this command.  This is not to display to a user, but instead used to uniquely identify this command */
+	FName GetCommandName() const
+	{
+		return UICommandInfo->GetCommandName();
+	}
+
+	/** Gets the text to send to the transaction system when creating an undo/redo event for this action */
+	FText GetUndoText() const
+	{
+		check( !bIsMode || !UndoText.IsEmpty() );	// All mode-based commands must have supplied UndoText.  Instantaneous commands handle their own undo/redo.
+		return UndoText;
+	}
+
+	/** Returns true if this is a mesh editing 'Mode' that the user will stay in to perform the action multiple times, or false
+		if the action applies instantly */
+	bool IsMode() const
+	{
+		return bIsMode;
+	}
+
+	/** Returns whether we rely on a hover location under the interactor being updated as we drag during this action */
+	bool NeedsHoverLocation() const
+	{
+		return bNeedsHoverLocation;
+	}
+
+	/** Returns whether this command will kick off regular free translation of the selected mesh elements when dragging starts */
+	bool NeedsDraggingInitiated() const
+	{
+		return bNeedsDraggingInitiated;
+	}
 
 	/** Gets the UI command info for this command */
 	const TSharedPtr<class FUICommandInfo>& GetUICommandInfo() const
@@ -24,14 +73,66 @@ public:
 		return UICommandInfo;
 	}
 
+	/** Creates an UI action for this command */
+	FUIAction MakeUIAction( class IMeshEditorModeUIContract& MeshEditorMode );
+
 
 protected:
 
+	/** The text to send to the transaction system when creating an undo / redo event for this action */
+	UPROPERTY( EditAnywhere, Category=MeshEditor )
+	FText UndoText;
+
+	/** True if this is a mesh editing 'Mode' that the user will stay in to perform the action multiple times, or false
+	    if the action applies instantly */
+	UPROPERTY( EditAnywhere, Category=MeshEditor )
+	bool bIsMode;
+
+	/** Whether this command will kick off regular free translation of the selected mesh elements when dragging starts */
+	UPROPERTY( EditAnywhere, Category=MeshEditor )
+	bool bNeedsDraggingInitiated;
+
+	/** Whether we rely on a hover location under the interactor being updated as we drag during this action */
+	UPROPERTY( EditAnywhere, Category=MeshEditor )
+	bool bNeedsHoverLocation;
+
 	/** Our UI command for this action */
 	TSharedPtr<FUICommandInfo> UICommandInfo;
 };
 
 
+UCLASS( abstract )
+class MESHEDITOR_API UMeshEditorVertexCommand : public UMeshEditorCommand
+{
+	GENERATED_BODY()
+
+public:
+
+	// Overrides
+	virtual EEditableMeshElementType GetElementType() const override
+	{
+		return EEditableMeshElementType::Vertex;
+	}
+
+};
+
+
+UCLASS( abstract )
+class MESHEDITOR_API UMeshEditorEdgeCommand : public UMeshEditorCommand
+{
+	GENERATED_BODY()
+
+public:
+
+	// Overrides
+	virtual EEditableMeshElementType GetElementType() const override
+	{
+		return EEditableMeshElementType::Edge;
+	}
+
+};
+
+
 UCLASS( abstract )
 class MESHEDITOR_API UMeshEditorPolygonCommand : public UMeshEditorCommand
 {
@@ -39,6 +140,12 @@ class MESHEDITOR_API UMeshEditorPolygonCommand : public UMeshEditorCommand
 
 public:
 
+	// Overrides
+	virtual EEditableMeshElementType GetElementType() const override
+	{
+		return EEditableMeshElementType::Polygon;
+	}
+
 };
 
 
@@ -139,9 +246,6 @@ public:
 	/** Sets the primary action to split edges and drag vertices */
 	TSharedPtr<FUICommandInfo> SplitEdgeAndDragVertex;
 
-	/** Sets the primary action to insert edge loops */
-	TSharedPtr<FUICommandInfo> InsertEdgeLoop;
-
 	/** Sets the primary action to extend edges */
 	TSharedPtr<FUICommandInfo> ExtendEdge;
 
diff --git a/Engine/Plugins/Editor/MeshEditor/Source/MeshEditor/Public/UndoHelpers.h b/Engine/Plugins/Editor/MeshEditor/Source/MeshEditor/Public/UndoHelpers.h
deleted file mode 100644
diff --git a/Engine/Plugins/Editor/MeshEditor/Source/MeshEditor/UndoHelpers.cpp b/Engine/Plugins/Editor/MeshEditor/Source/MeshEditor/UndoHelpers.cpp
deleted file mode 100644
diff --git a/Engine/Plugins/Editor/MeshEditor/Source/PolygonModeling/InsertEdgeLoop.cpp b/Engine/Plugins/Editor/MeshEditor/Source/PolygonModeling/InsertEdgeLoop.cpp
new file mode 100644
--- /dev/null
+++ b/Engine/Plugins/Editor/MeshEditor/Source/PolygonModeling/InsertEdgeLoop.cpp
@@ -0,0 +1,116 @@
+// Copyright 1998-2016 Epic Games, Inc. All Rights Reserved.
+
+#include "InsertEdgeLoop.h"
+#include "IMeshEditorModeEditingContract.h"
+#include "IMeshEditorModeUIContract.h"
+#include "UICommandInfo.h"
+#include "EditableMesh.h"
+#include "MeshElement.h"
+#include "MultiBoxBuilder.h"
+#include "UICommandList.h"
+
+
+#define LOCTEXT_NAMESPACE "MeshEditorMode"
+
+
+void UInsertEdgeLoopCommand::RegisterUICommand( FBindingContext* BindingContext )
+{
+	UI_COMMAND_EXT( BindingContext, /* Out */ UICommandInfo, "InsertEdgeLoop", "Insert Edge Loop Mode", "Set the primary action to insert edge loops.", EUserInterfaceActionType::RadioButton, FInputChord() );
+}
+
+
+void UInsertEdgeLoopCommand::ApplyDuringDrag( IMeshEditorModeEditingContract& MeshEditorMode, bool& bOutShouldDeselectAllFirst, TArray<FMeshElement>& OutMeshElementsToSelect )
+{
+	// Insert edge loop
+	static TMap< UEditableMesh*, TArray< FMeshElement > > SelectedMeshesAndEdges;
+	MeshEditorMode.GetSelectedMeshesAndEdges( /* Out */ SelectedMeshesAndEdges );
+
+	if( SelectedMeshesAndEdges.Num() > 0 )
+	{
+		// Deselect the edges first, since they'll be deleted or split up while inserting the edge loop,
+		// and we want them to be re-selected after undo
+		MeshEditorMode.DeselectMeshElements( SelectedMeshesAndEdges );
+
+
+		for( auto& MeshAndEdges : SelectedMeshesAndEdges )
+		{
+			UEditableMesh* EditableMesh = MeshAndEdges.Key;
+			verify( !EditableMesh->AnyChangesToUndo() );
+
+			const TArray<FMeshElement>& EdgeElements = MeshAndEdges.Value;
+
+			// Figure out where to add the loop along the edge
+			static TArray<float> Splits;
+			MeshEditorMode.FindEdgeSplitUnderInteractor( MeshEditorMode.GetActiveActionInteractor(), EditableMesh, EdgeElements, /* Out */ Splits );
+
+			// Insert the edge loop
+			if( Splits.Num() > 0 )
+			{
+				// @todo mesheditor edgeloop: Test code
+				if( false )
+				{
+					if( Splits[ 0 ] > 0.25f )
+					{
+						Splits.Insert( FMath::Max( Splits[ 0 ] - 0.2f, 0.0f ), 0 );
+					}
+					if( Splits.Last() < 0.75f )
+					{
+						Splits.Add( FMath::Min( Splits.Last() + 0.2f, 1.0f ) );
+					}
+				}
+
+				for( const FMeshElement& EdgeMeshElement : EdgeElements )
+				{
+					const FEdgeID EdgeID( EdgeMeshElement.ElementAddress.ElementID );
+
+					static TArray<FEdgeID> NewEdgeIDs;
+					NewEdgeIDs.Reset();
+
+					EditableMesh->InsertEdgeLoop( EdgeID, Splits, /* Out */ NewEdgeIDs );
+
+					// Select all of the new edges that were created by inserting the loop
+					if( NewEdgeIDs.Num() > 0 )
+					{
+						bOutShouldDeselectAllFirst = true;	// Don't keep the edge selected
+
+						for( const FEdgeID NewEdgeID : NewEdgeIDs )
+						{
+							FMeshElement MeshElementToSelect;
+							{
+								MeshElementToSelect.Component = EdgeMeshElement.Component;
+								MeshElementToSelect.ElementAddress.SubMeshAddress = EdgeMeshElement.ElementAddress.SubMeshAddress;
+								MeshElementToSelect.ElementAddress.ElementType = EEditableMeshElementType::Edge;
+								MeshElementToSelect.ElementAddress.ElementID = NewEdgeID;
+							}
+
+							// Queue selection of this new element.  We don't want it to be part of the current action.
+							OutMeshElementsToSelect.Add( MeshElementToSelect );
+						}
+					}
+				}
+
+				MeshEditorMode.TrackUndo( EditableMesh, EditableMesh->MakeUndo() );
+			}
+		}
+	}
+}
+
+
+void UInsertEdgeLoopCommand::AddToVRRadialMenuActionsMenu( IMeshEditorModeUIContract& MeshEditorMode, FMenuBuilder& MenuBuilder, TSharedPtr<FUICommandList> CommandList, const FName TEMPHACK_StyleSetName, class UVREditorMode* VRMode )
+{
+	if( MeshEditorMode.GetMeshElementSelectionMode() == EEditableMeshElementType::Edge )
+	{
+		MenuBuilder.AddMenuEntry(
+			LOCTEXT("InsertEdgeLoop", "Insert Loop"),
+			FText(),
+			FSlateIcon(TEMPHACK_StyleSetName, "MeshEditorMode.EdgeInsert"),	// @todo mesheditor extensibility: TEMPHACK for style; Need PolygonModelingStyle, probably.  Or we're just cool with exporting MeshEditorModeStyle, since we're all the same plugin after all.
+			MakeUIAction( MeshEditorMode ),
+			NAME_None,
+			EUserInterfaceActionType::ToggleButton
+			);
+	}
+}
+
+
+#undef LOCTEXT_NAMESPACE
+
diff --git a/Engine/Plugins/Editor/MeshEditor/Source/PolygonModeling/InsertEdgeLoop.h b/Engine/Plugins/Editor/MeshEditor/Source/PolygonModeling/InsertEdgeLoop.h
new file mode 100644
--- /dev/null
+++ b/Engine/Plugins/Editor/MeshEditor/Source/PolygonModeling/InsertEdgeLoop.h
@@ -0,0 +1,30 @@
+// Copyright 1998-2016 Epic Games, Inc. All Rights Reserved.
+
+#pragma once
+
+#include "MeshEditorCommands.h"
+#include "InsertEdgeLoop.generated.h"
+
+
+/** With an edge select, inserts a loop of edge perpendicular to that edge while dragging */
+UCLASS()
+class UInsertEdgeLoopCommand : public UMeshEditorEdgeCommand
+{
+	GENERATED_BODY()
+
+protected:
+
+	UInsertEdgeLoopCommand()
+	{
+		bIsMode = true;
+		UndoText = NSLOCTEXT( "MeshEditor", "UndoInsertEdgeLoop", "Insert Edge Loop" );
+		bNeedsHoverLocation = true;
+		bNeedsDraggingInitiated = false;
+	}
+
+	// Overrides
+	virtual void RegisterUICommand( class FBindingContext* BindingContext ) override;
+	virtual void ApplyDuringDrag( IMeshEditorModeEditingContract& MeshEditorMode, bool& bOutShouldDeselectAllFirst, TArray<FMeshElement>& OutMeshElementsToSelect ) override;
+	virtual void AddToVRRadialMenuActionsMenu( class IMeshEditorModeUIContract& MeshEditorMode, class FMenuBuilder& MenuBuilder, TSharedPtr<FUICommandList> CommandList, const FName TEMPHACK_StyleSetName, class UVREditorMode* VRMode ) override;
+
+};
diff --git a/Engine/Plugins/Editor/MeshEditor/Source/PolygonModeling/PolygonModeling.Build.cs b/Engine/Plugins/Editor/MeshEditor/Source/PolygonModeling/PolygonModeling.Build.cs
--- a/Engine/Plugins/Editor/MeshEditor/Source/PolygonModeling/PolygonModeling.Build.cs
+++ b/Engine/Plugins/Editor/MeshEditor/Source/PolygonModeling/PolygonModeling.Build.cs
@@ -12,9 +12,11 @@ namespace UnrealBuildTool.Rules
 				    "CoreUObject",
 					"Slate",
 					"Engine",
-				    "UnrealEd",
+					"InputCore",
+					"UnrealEd",
 					"MeshEditingRuntime",
-					"MeshEditor"
+					"MeshEditor",
+					"SlateCore"
 				}
 			);
 		}
diff --git a/Engine/Plugins/Editor/MeshEditor/Source/PolygonModeling/TessellatePolygons.cpp b/Engine/Plugins/Editor/MeshEditor/Source/PolygonModeling/TessellatePolygons.cpp
--- a/Engine/Plugins/Editor/MeshEditor/Source/PolygonModeling/TessellatePolygons.cpp
+++ b/Engine/Plugins/Editor/MeshEditor/Source/PolygonModeling/TessellatePolygons.cpp
@@ -2,7 +2,6 @@
 
 #include "TessellatePolygons.h"
 #include "IMeshEditorModeEditingContract.h"
-#include "UndoHelpers.h"
 #include "UICommandInfo.h"
 #include "EditableMesh.h"
 #include "MeshElement.h"
@@ -13,14 +12,13 @@
 
 void UTessellatePolygonsCommand::RegisterUICommand( FBindingContext* BindingContext )
 {
-	UI_COMMAND_EXT( BindingContext, UICommandInfo, "Tessellate Selected Polygons", "Tessellate selected polygons into smaller polygons.", EUserInterfaceActionType::Button, FInputChord() );
+	UI_COMMAND_EXT( BindingContext, /* Out */ UICommandInfo, "TessellatePolygons", "Tessellate Selected Polygons", "Tessellate selected polygons into smaller polygons.", EUserInterfaceActionType::Button, FInputChord() );
 }
 
 
-void UTessellatePolygonsCommand::Execute( IMeshEditorModeEditingContract& MeshEditorMode, bool& bOutWasSuccessful )
+void UTessellatePolygonsCommand::Execute( IMeshEditorModeEditingContract& MeshEditorMode )
 {
-	bOutWasSuccessful = false;
-	if( MeshEditorMode.GetActiveAction() != IMeshEditorModeEditingContract::EMeshEditAction::None )
+	if( MeshEditorMode.GetActiveAction() != EMeshEditAction::None )
 	{
 		return;
 	}
@@ -85,7 +83,7 @@ void UTessellatePolygonsCommand::Execute( IMeshEditorModeEditingContract& MeshEd
 
 		EditableMesh->EndModification();
 
-		UndoHelpers::StoreUndo( EditableMesh, EditableMesh->MakeUndo() );
+		MeshEditorMode.TrackUndo( EditableMesh, EditableMesh->MakeUndo() );
 	}
 
 	// Make sure we're not still hovering over a polygon we removed
@@ -93,8 +91,6 @@ void UTessellatePolygonsCommand::Execute( IMeshEditorModeEditingContract& MeshEd
 
 	// Select the new smaller polygons
 	MeshEditorMode.SelectMeshElements( MeshElementsToSelect );
-
-	bOutWasSuccessful = true;
 }
 
 
diff --git a/Engine/Plugins/Editor/MeshEditor/Source/PolygonModeling/TessellatePolygons.h b/Engine/Plugins/Editor/MeshEditor/Source/PolygonModeling/TessellatePolygons.h
--- a/Engine/Plugins/Editor/MeshEditor/Source/PolygonModeling/TessellatePolygons.h
+++ b/Engine/Plugins/Editor/MeshEditor/Source/PolygonModeling/TessellatePolygons.h
@@ -16,6 +16,6 @@ protected:
 
 	// Overrides
 	virtual void RegisterUICommand( class FBindingContext* BindingContext ) override;
-	virtual void Execute( class IMeshEditorModeEditingContract& MeshEditorMode, bool& bOutWasSuccessful ) override;
+	virtual void Execute( class IMeshEditorModeEditingContract& MeshEditorMode ) override;
 
 };
diff --git a/Engine/Source/Runtime/Slate/Public/Framework/Commands/Commands.h b/Engine/Source/Runtime/Slate/Public/Framework/Commands/Commands.h
--- a/Engine/Source/Runtime/Slate/Public/Framework/Commands/Commands.h
+++ b/Engine/Source/Runtime/Slate/Public/Framework/Commands/Commands.h
@@ -15,10 +15,10 @@
 SLATE_API void UI_COMMAND_Function( FBindingContext* This, TSharedPtr< FUICommandInfo >& OutCommand, const TCHAR* OutSubNamespace, const TCHAR* OutCommandName, const TCHAR* OutCommandNameUnderscoreTooltip, const ANSICHAR* DotOutCommandName, const TCHAR* FriendlyName, const TCHAR* InDescription, const EUserInterfaceActionType::Type CommandType, const FInputChord& InDefaultChord );
 
 // This macro requires LOCTEXT_NAMESPACE to be defined. If you don't want the command to be placed under a sub namespace, provide "" as the namespace.
-#define UI_COMMAND_EXT( BindingContext, CommandId, FriendlyName, InDescription, CommandType, InDefaultChord ) \
-	UI_COMMAND_Function( BindingContext, CommandId, TEXT(LOCTEXT_NAMESPACE), TEXT(#CommandId), TEXT(#CommandId) TEXT("_ToolTip"), "." #CommandId, TEXT(FriendlyName), TEXT(InDescription), CommandType, InDefaultChord );
+#define UI_COMMAND_EXT( BindingContext, OutUICommandInfo, CommandIdName, FriendlyName, InDescription, CommandType, InDefaultChord ) \
+	UI_COMMAND_Function( BindingContext, OutUICommandInfo, TEXT(LOCTEXT_NAMESPACE), TEXT(CommandIdName), TEXT(CommandIdName) TEXT("_ToolTip"), "." CommandIdName, TEXT(FriendlyName), TEXT(InDescription), CommandType, InDefaultChord );
 #define UI_COMMAND( CommandId, FriendlyName, InDescription, CommandType, InDefaultChord ) \
-	UI_COMMAND_EXT( this, CommandId, FriendlyName, InDescription, CommandType, InDefaultChord )
+	UI_COMMAND_EXT( this, CommandId, #CommandId, FriendlyName, InDescription, CommandType, InDefaultChord )
 	
 #undef LOC_DEFINE_REGION
 
