Mesh Editor extensibility (phase 3)
- Inset Polygon and Bevel Polygon have been migrated to a new module
- Mesh editing actions that require custom logic before starting a drag can now be built in external modules
- Added 'FindEditableMesh' for cases where we are expecting a mesh to be cached already
- Added 'IsMeshElementSelected' and replaced some uses of 'GetSelectedMeshElementIndex'
- Some additional features for MeshEditorMode have been made public to allow for this.  For example, the currently hovered mesh element can now be queried using 'GetHoveredMeshElement'
- Renamed 'TessellatePolygons' to 'TessellatePolygon' for consistency
- A few bug fixes to previous extensibility checkin

diff --git a/Engine/Plugins/Editor/MeshEditor/Source/MeshEditor/MeshEditorCommands.cpp b/Engine/Plugins/Editor/MeshEditor/Source/MeshEditor/MeshEditorCommands.cpp
--- a/Engine/Plugins/Editor/MeshEditor/Source/MeshEditor/MeshEditorCommands.cpp
+++ b/Engine/Plugins/Editor/MeshEditor/Source/MeshEditor/MeshEditorCommands.cpp
@@ -146,8 +146,6 @@ void FMeshEditorPolygonCommands::RegisterCommands()
 	UI_COMMAND(MovePolygon, "Move Polygon Mode", "Set the primary action to move polygons.", EUserInterfaceActionType::RadioButton, FInputChord(EKeys::F1));
 	UI_COMMAND(ExtrudePolygon, "Extrude Polygon Mode", "Set the primary action to extrude polygons.", EUserInterfaceActionType::RadioButton, FInputChord(EKeys::F2));
 	// UI_COMMAND(FreelyExtrudePolygon, "Freely Extrude Polygon Mode", "Set the primary action to freely extrude polygons.", EUserInterfaceActionType::RadioButton, FInputChord(EKeys::F3));
-	UI_COMMAND(InsetPolygon, "Inset Polygon Mode", "Set the primary action to inset polygons.", EUserInterfaceActionType::RadioButton, FInputChord(EKeys::F3));
-	UI_COMMAND(BevelPolygon, "Bevel Polygon Mode", "Set the primary action to bevel polygons.", EUserInterfaceActionType::RadioButton, FInputChord(EKeys::F4));
 
 	UI_COMMAND(FlipPolygon, "Flip Polygon", "Flip the currently selected polygons.", EUserInterfaceActionType::Button, FInputChord(EKeys::F, EModifierKey::Shift));
 	UI_COMMAND(TriangulatePolygon, "Triangulate Polygon", "Triangulate the currently selected polygons.", EUserInterfaceActionType::Button, FInputChord(EKeys::T));
diff --git a/Engine/Plugins/Editor/MeshEditor/Source/MeshEditor/MeshEditorMode.cpp b/Engine/Plugins/Editor/MeshEditor/Source/MeshEditor/MeshEditorMode.cpp
--- a/Engine/Plugins/Editor/MeshEditor/Source/MeshEditor/MeshEditorMode.cpp
+++ b/Engine/Plugins/Editor/MeshEditor/Source/MeshEditor/MeshEditorMode.cpp
@@ -51,8 +51,6 @@ namespace EMeshEditAction
 	const Type SplitEdgeAndDragVertex( "SplitEdgeAndDragVertex" );
 	const Type ExtrudePolygon( "ExtrudePolygon" );
 	const Type FreelyExtrudePolygon( "FreelyExtrudePolygon" );
-	const Type InsetPolygon( "InsetPolygon" );
-	const Type BevelPolygon( "BevelPolygon" );
 	const Type ExtendEdge( "ExtendEdge" );
 	const Type ExtendVertex( "ExtendVertex" );
 	const Type EditVertexCornerSharpness( "EditVertexCornerSharpness" );
@@ -144,7 +142,7 @@ TUniquePtr<FChange> FMeshEditorMode::FSelectOrDeselectMeshElementsChange::Execut
 					if( CurrentElementSelectionMode == EEditableMeshElementType::Any ||
 					    MeshElementToSelect.ElementAddress.ElementType == CurrentElementSelectionMode )
 					{
-						UEditableMesh* EditableMesh = MeshEditorMode.FindOrCreateEditableMesh( *MeshElementToSelect.Component, MeshElementToSelect.ElementAddress.SubMeshAddress );
+						const UEditableMesh* EditableMesh = MeshEditorMode.FindEditableMesh( *MeshElementToSelect.Component, MeshElementToSelect.ElementAddress.SubMeshAddress );
 						if( EditableMesh != nullptr )
 						{
 							if( IsElementIDValid( MeshElementToSelect, EditableMesh ) )
@@ -253,7 +251,7 @@ TUniquePtr<FChange> FMeshEditorMode::FSetElementSelectionModeChange::Execute( UO
 
 			if( Component )
 			{
-				UEditableMesh* EditableMesh = MeshEditorMode.FindOrCreateEditableMesh( *Component, SubMeshAddress );
+				const UEditableMesh* EditableMesh = MeshEditorMode.FindEditableMesh( *Component, SubMeshAddress );
 
 				if( Input.Mode == EEditableMeshElementType::Vertex )
 				{
@@ -398,7 +396,6 @@ FMeshEditorMode::FMeshEditorMode()
 	  bIsFirstActiveActionUpdate( false ),
 	  ExtrudePolygonAxisOrigin( FVector::ZeroVector ),
 	  ExtrudePolygonAxisDirection( FVector::ZeroVector ),
-	  InsetUsingPolygonElement(),
 	  SplitEdgeMeshesAndEdgesToSplit(),
 	  SplitEdgeSplitList(),
 	  EditSharpnessStartLocation( FVector::ZeroVector ),
@@ -546,8 +543,6 @@ void FMeshEditorMode::BindCommands()
 	RegisterPolygonEditingMode( MeshEditorPolygonActions.MovePolygon, EMeshEditAction::Move );
 	RegisterPolygonEditingMode( MeshEditorPolygonActions.ExtrudePolygon, EMeshEditAction::ExtrudePolygon );
 //	RegisterPolygonEditingMode( MeshEditorPolygonActions.FreelyExtrudePolygon, EMeshEditAction::FreelyExtrudePolygon );	// @todo mesheditor: "Freely Extrude" disabled until we have proper gizmo/snap support
-	RegisterPolygonEditingMode( MeshEditorPolygonActions.InsetPolygon, EMeshEditAction::InsetPolygon );
-	RegisterPolygonEditingMode( MeshEditorPolygonActions.BevelPolygon, EMeshEditAction::BevelPolygon );
 
 	RegisterCommonEditingMode( MeshEditorCommonActions.DrawVertices, EMeshEditAction::DrawVertices );
 
@@ -810,6 +805,21 @@ void FMeshEditorMode::Exit()
 }
 
 
+const UEditableMesh* FMeshEditorMode::FindEditableMesh( UPrimitiveComponent& Component, const FEditableMeshSubMeshAddress& SubMeshAddress ) const
+{
+	const UEditableMesh* EditableMesh = nullptr;
+
+	// Grab the existing editable mesh from our cache if we have one, otherwise create one now
+	UEditableMesh* const* EditableMeshPtr = CachedEditableMeshes.Find( SubMeshAddress );
+	if( EditableMeshPtr != nullptr )
+	{
+		EditableMesh = *EditableMeshPtr;
+	}
+
+	return EditableMesh;
+}
+
+
 UEditableMesh* FMeshEditorMode::FindOrCreateEditableMesh( UPrimitiveComponent& Component, const FEditableMeshSubMeshAddress& SubMeshAddress )
 {
 	UEditableMesh* EditableMesh = nullptr;
@@ -905,7 +915,7 @@ void FMeshEditorMode::Tick( FEditorViewportClient* ViewportClient, float DeltaTi
 		FMeshEditorInteractorData& MeshEditorInteractorData = GetMeshEditorInteractorData( ActiveActionInteractor );
 
 		// If not already selected, add it to our selection set
-		if( GetSelectedMeshElementIndex( MeshEditorInteractorData.HoveredMeshElement ) == INDEX_NONE )
+		if( !IsMeshElementSelected( MeshEditorInteractorData.HoveredMeshElement ) )
 		{
 			// Only add elements of the same type.  Otherwise it would just cause things to become deselected as you move between
 			// different element types, as we don't allow you to select elements that have overlapping geometry
@@ -1105,7 +1115,8 @@ void FMeshEditorMode::CommitEditableMeshIfNecessary( UEditableMesh* EditableMesh
 		}
 
 		// @todo mesheditor: this is a little bit fragile. Ideally we initialize these things after the new instance has been created.
-		FixUpMeshElement( InsetUsingPolygonElement );
+		// @todo mesheditor extensibility: Figure out how external FMeshElements can be fixed up (either a callback with FixUpMeshElement function access, or they are registered with this system?)
+		// FixUpMeshElement( InsetUsingPolygonElement );
 
 		for( auto& SplitEdgeMeshAndEdgesToSplit : SplitEdgeMeshesAndEdgesToSplit )
 		{
@@ -2780,7 +2791,7 @@ void FMeshEditorMode::Render( const FSceneView* SceneView, FViewport* Viewport,
 }
 
 
-FMeshEditorMode::FMeshEditorInteractorData& FMeshEditorMode::GetMeshEditorInteractorData( UViewportInteractor* ViewportInteractor )
+FMeshEditorMode::FMeshEditorInteractorData& FMeshEditorMode::GetMeshEditorInteractorData( UViewportInteractor* ViewportInteractor ) const
 {
 	check( ViewportInteractor != nullptr );
 
@@ -3242,84 +3253,6 @@ void FMeshEditorMode::UpdateActiveAction( const bool bIsActionFinishing )
 		OutExtrudeDistance = BestExtrudeDistance;
 	};
 
-	// Figures out how far we should inset on a polygon by figuring out the maximum distance from all of the polygon's perimeter
-	// vertices to the center of the polygon, then subtracting the distance between the hovered impact point and polygon center.
-	auto FindInsetAmount = [](
-		FMeshEditorMode* This,
-		UViewportInteractor* ViewportInteractor,
-		const FPolygonRef& PolygonRef,
-		UPrimitiveComponent& Component,
-		const UEditableMesh* EditableMesh,
-		float& OutInsetFixedDistance,
-		float& OutInsetProgressTowardCenter )
-	{
-		// @todo grabber: Glitches out when using grabber sphere near the corner of an inset polygon.
-		check( ViewportInteractor != nullptr );
-
-		OutInsetFixedDistance = 0.0f;
-		OutInsetProgressTowardCenter = 0.0f;
-
-		FMeshEditorInteractorData& MeshEditorInteractorData = This->GetMeshEditorInteractorData( ViewportInteractor );
-		if( MeshEditorInteractorData.bLaserIsValid || MeshEditorInteractorData.bGrabberSphereIsValid )
-		{
-			const FMatrix ComponentToWorldMatrix = Component.GetRenderMatrix();
-
-			const FPlane PolygonPlane = EditableMesh->ComputePolygonPlane( PolygonRef );
-			const FVector PolygonCenter = EditableMesh->ComputePolygonCenter( PolygonRef );
-
-			// @todo mesheditor: no longer supports grabber sphere. Add this back in (plane vs sphere closest point check) if we want to support it.
-			const FVector ComponentSpaceRayStart = ComponentToWorldMatrix.InverseTransformPosition( MeshEditorInteractorData.LaserStart );
-			const FVector ComponentSpaceRayEnd = ComponentToWorldMatrix.InverseTransformPosition( MeshEditorInteractorData.LaserEnd );
-
-			FVector ComponentSpaceRayIntersectionWithPolygonPlane;
-			if( FMath::SegmentPlaneIntersection( ComponentSpaceRayStart, ComponentSpaceRayEnd, PolygonPlane, /* Out */ ComponentSpaceRayIntersectionWithPolygonPlane ) )
-			{
-				// @todo mesheditor debug
-				// DrawDebugSphere( This->GetWorld(), ComponentToWorldMatrix.TransformPosition( ComponentSpaceRayIntersectionWithPolygonPlane ), 1.5f, 32, FColor::Red, false, 0.0f );
-
-				bool bFoundTriangle = false;
-				float CenterWeight = 0.0f;
-				FVector BestEdgeVertex0Position, BestEdgeVertex1Position;
-				{
-					static TArray<FVertexID> PerimeterVertexIDs;
-					EditableMesh->GetPolygonPerimeterVertices( PolygonRef, /* Out */ PerimeterVertexIDs );
-
-					for( int32 VertexNumber = 0; VertexNumber < PerimeterVertexIDs.Num(); ++VertexNumber )
-					{
-						const int32 NextVertexNumber = ( VertexNumber + 1 ) % PerimeterVertexIDs.Num();
-
-
-						const FVertexID EdgeVertex0 = PerimeterVertexIDs[ VertexNumber ];
-						const FVertexID EdgeVertex1 = PerimeterVertexIDs[ NextVertexNumber ];
-
-						const FVector EdgeVertex0Position = EditableMesh->GetVertexAttribute( EdgeVertex0, UEditableMeshAttribute::VertexPosition(), 0 );
-						const FVector EdgeVertex1Position = EditableMesh->GetVertexAttribute( EdgeVertex1, UEditableMeshAttribute::VertexPosition(), 0 );
-
-						BestEdgeVertex0Position = EdgeVertex0Position;
-						BestEdgeVertex1Position = EdgeVertex1Position;
-
-						const FVector VertexWeights = FMath::ComputeBaryCentric2D( ComponentSpaceRayIntersectionWithPolygonPlane, EdgeVertex0Position, EdgeVertex1Position, PolygonCenter );
-						if( VertexWeights.X >= 0.0f && VertexWeights.Y >= 0.0f && VertexWeights.Z >= 0.0f )
-						{
-							CenterWeight = VertexWeights.Z;
-
-							bFoundTriangle = true;
-							break;
-						}
-					}
-				}
-
-				if( bFoundTriangle )
-				{
-					// @todo mesheditor debug
-					// DrawDebugLine( This->GetWorld(), ComponentToWorldMatrix.TransformPosition( BestEdgeVertex0Position ), ComponentToWorldMatrix.TransformPosition( BestEdgeVertex1Position ), FColor::Yellow, false, 0.0f, 0, 2.0f );
-
-					OutInsetProgressTowardCenter = CenterWeight;
-				}
-			}
-		}
-	};
-
 
 	// If this is an interim change, then everything that happens here -- all changes to our meshes, and even selection
 	// changes -- are guaranteed to be rolled back at the beginning of the next frame.  So we'll intercept any requests
@@ -3589,144 +3522,6 @@ void FMeshEditorMode::UpdateActiveAction( const bool bIsActionFinishing )
 			}
 		}
 	}
-	else if( ActiveAction == EMeshEditAction::InsetPolygon ||
-			 ActiveAction == EMeshEditAction::BevelPolygon )
-	{
-		static TMap< UEditableMesh*, TArray< FMeshElement > > MeshesWithPolygonsToInset;
-		GetSelectedMeshesAndPolygons( /* Out */ MeshesWithPolygonsToInset );
-
-		if( MeshesWithPolygonsToInset.Num() > 0 )
-		{
-			if( InsetUsingPolygonElement.IsValidMeshElement() &&
-				InsetUsingPolygonElement.ElementAddress.ElementType == EEditableMeshElementType::Polygon &&
-				GetSelectedMeshElementIndex( InsetUsingPolygonElement ) != INDEX_NONE )
-			{
-				UPrimitiveComponent* InsetUsingComponent = InsetUsingPolygonElement.Component.Get();
-				check( InsetUsingComponent != nullptr );
-
-				UEditableMesh* InsetUsingEditableMesh = FindOrCreateEditableMesh( *InsetUsingComponent, InsetUsingPolygonElement.ElementAddress.SubMeshAddress );
-				if( InsetUsingEditableMesh != nullptr )
-				{
-					const FPolygonRef InsetUsingPolygonRef( InsetUsingPolygonElement.ElementAddress.SectionID, FPolygonID( InsetUsingPolygonElement.ElementAddress.ElementID ) );
-
-					// Figure out how far to inset the polygon
-					float InsetFixedDistance = 0.0f;
-					float InsetProgressTowardCenter = 0.0f;
-					FindInsetAmount(
-						this,
-						ActiveActionInteractor,
-						InsetUsingPolygonRef,
-						*InsetUsingComponent,
-						InsetUsingEditableMesh,
-						/* Out */ InsetFixedDistance,
-						/* Out */ InsetProgressTowardCenter );
-
-					if( InsetFixedDistance > SMALL_NUMBER ||
-						InsetProgressTowardCenter > SMALL_NUMBER )
-					{
-
-						// Deselect the mesh elements before we delete them.  This will make sure they become selected again after undo.
-						{
-							FCompoundChangeInput CompoundRevertInput;
-
-							FSelectOrDeselectMeshElementsChangeInput ChangeInput;
-							for( auto& MeshAndPolygons : MeshesWithPolygonsToInset )
-							{
-								UEditableMesh* EditableMesh = MeshAndPolygons.Key;
-								const TArray<FMeshElement>& PolygonsToDeselect = MeshAndPolygons.Value;
-								ChangeInput.MeshElementsToDeselect.Append( PolygonsToDeselect );
-							}
-
-							CompoundRevertInput.Subchanges.Add( MoveTemp( FSelectOrDeselectMeshElementsChange( MoveTemp( ChangeInput ) ).Execute( MeshEditorModeProxyObject ) ) );
-
-							TrackUndo( MeshEditorModeProxyObject, MakeUnique<FCompoundChange>( MoveTemp( CompoundRevertInput ) ) );
-						}
-
-						for( auto& MeshAndPolygons : MeshesWithPolygonsToInset )
-						{
-							UEditableMesh* EditableMesh = MeshAndPolygons.Key;
-							const TArray<FMeshElement>& PolygonsToInset = MeshAndPolygons.Value;
-
-							UPrimitiveComponent* Component = PolygonsToInset[ 0 ].Component.Get();	// NOTE: All polygons in this array belong to the same mesh/component, so we just need the first element
-							check( Component != nullptr );
-
-
-							static TArray<FPolygonRef> PolygonRefsToInset;
-							PolygonRefsToInset.Reset();
-							for( const FMeshElement& PolygonToInset : PolygonsToInset )
-							{
-								FPolygonRef PolygonRef( PolygonToInset.ElementAddress.SectionID, FPolygonID( PolygonToInset.ElementAddress.ElementID ) );
-								PolygonRefsToInset.Add( PolygonRef );
-							}
-
-							{
-								verify( !EditableMesh->AnyChangesToUndo() );
-
-								// Inset time!!
-								static TArray<FPolygonRef> NewCenterInsetPolygons;
-								static TArray<FPolygonRef> NewSideInsetPolygons;
-								if( ActiveAction == EMeshEditAction::InsetPolygon )
-								{
-									const EInsetPolygonsMode InsetPolygonsMode = EInsetPolygonsMode::All;	// @todo mesheditor inset: Make configurable?
-																											// @todo mesheditor inset: Add options for Fixed distance (instead of Percentage distance, like now.)
-
-									EditableMesh->InsetPolygons( PolygonRefsToInset, InsetFixedDistance, InsetProgressTowardCenter, InsetPolygonsMode, /* Out */ NewCenterInsetPolygons, /* Out */ NewSideInsetPolygons );
-								}
-								else if( ensure( ActiveAction == EMeshEditAction::BevelPolygon ) )
-								{
-									EditableMesh->BevelPolygons( PolygonRefsToInset, InsetFixedDistance, InsetProgressTowardCenter, /* Out */ NewCenterInsetPolygons, /* Out */ NewSideInsetPolygons );
-								}
-
-								// Make sure the new polygons are selected.  The old polygon was deleted and will become deselected automatically.
-								if( NewCenterInsetPolygons.Num() > 0 )
-								{
-									for( int32 NewPolygonNumber = 0; NewPolygonNumber < NewCenterInsetPolygons.Num(); ++NewPolygonNumber )
-									{
-										const FPolygonRef& NewInsetCenterPolygon = NewCenterInsetPolygons[ NewPolygonNumber ];
-
-										const FMeshElement& MeshElement = PolygonsToInset[ NewPolygonNumber ];
-
-										FMeshElement PolygonMeshElement;
-										{
-											PolygonMeshElement.Component = MeshElement.Component;
-											PolygonMeshElement.ElementAddress = MeshElement.ElementAddress;
-											PolygonMeshElement.ElementAddress.SectionID = NewInsetCenterPolygon.SectionID;
-											PolygonMeshElement.ElementAddress.ElementID = NewInsetCenterPolygon.PolygonID;
-										}
-
-										// Queue selection of this new element.  We don't want it to be part of the current action.
-										MeshElementsToSelect.Add( PolygonMeshElement );
-									}
-								}
-								else if( ensure( NewSideInsetPolygons.Num() > 0 ) )
-								{
-									for( int32 NewPolygonNumber = 0; NewPolygonNumber < NewSideInsetPolygons.Num(); ++NewPolygonNumber )
-									{
-										const FPolygonRef& NewInsetSidePolygon = NewSideInsetPolygons[ NewPolygonNumber ];
-
-										const FMeshElement& MeshElement = PolygonsToInset[ NewPolygonNumber ];
-
-										FMeshElement PolygonMeshElement;
-										{
-											PolygonMeshElement.Component = MeshElement.Component;
-											PolygonMeshElement.ElementAddress = MeshElement.ElementAddress;
-											PolygonMeshElement.ElementAddress.SectionID = NewInsetSidePolygon.SectionID;
-											PolygonMeshElement.ElementAddress.ElementID = NewInsetSidePolygon.PolygonID;
-										}
-
-										// Queue selection of this new element.  We don't want it to be part of the current action.
-										MeshElementsToSelect.Add( PolygonMeshElement );
-									}
-								}
-							}
-
-							TrackUndo( EditableMesh, EditableMesh->MakeUndo() );
-						}
-					}
-				}
-			}
-		}
-	}
 	else if( ActiveAction == EMeshEditAction::ExtendEdge )
 	{
 		// Extend edge
@@ -4262,7 +4057,7 @@ void FMeshEditorMode::UpdateActiveAction( const bool bIsActionFinishing )
 			{
 				if( ActiveAction == CommandCDO->GetCommandName() )
 				{
-					CommandCDO->ApplyDuringDrag( *this, /* Out */ bShouldDeselectAllFirst, /* Out */ MeshElementsToSelect );
+					CommandCDO->ApplyDuringDrag( *this, ActiveActionInteractor, /* Out */ bShouldDeselectAllFirst, /* Out */ MeshElementsToSelect );
 
 					// Should always only be one candidate
 					bFoundValidCommand = true;
@@ -5126,7 +4921,7 @@ void FMeshEditorMode::OnViewportInteractionInputAction( FEditorViewportClient& V
 								PolygonElement.ElementAddress.ElementType == EEditableMeshElementType::Polygon )
 							{
 								// Is it selected?
-								if( GetSelectedMeshElementIndex( PolygonElement ) != INDEX_NONE )
+								if( IsMeshElementSelected( PolygonElement ) )
 								{
 									UEditableMesh* EditableMesh = FindOrCreateEditableMesh( *PolygonElement.Component, PolygonElement.ElementAddress.SubMeshAddress );
 									if( EditableMesh != nullptr )
@@ -5169,46 +4964,6 @@ void FMeshEditorMode::OnViewportInteractionInputAction( FEditorViewportClient& V
 						const bool bActionNeedsHoverLocation = false;
 						StartAction( EMeshEditAction::FreelyExtrudePolygon, ViewportInteractor, bActionNeedsHoverLocation, LOCTEXT( "UndoFreelyExtrudePolygon", "Freely Extrude Polygon" ) );
 					}
-					else if( SelectedMeshElementType == EEditableMeshElementType::Polygon && 
-							 ( EquippedPolygonAction == EMeshEditAction::InsetPolygon || EquippedPolygonAction == EMeshEditAction::BevelPolygon ) )
-					{
-						this->InsetUsingPolygonElement = FMeshElement();
-						{
-							const FMeshElement& PolygonElement = MeshEditorInteractorData.HoveredMeshElement;
-
-							if( PolygonElement.IsValidMeshElement() &&
-								PolygonElement.ElementAddress.ElementType == EEditableMeshElementType::Polygon )
-							{
-								// Is it selected?
-								if( GetSelectedMeshElementIndex( PolygonElement ) != INDEX_NONE )
-								{
-									UEditableMesh* EditableMesh = FindOrCreateEditableMesh( *PolygonElement.Component, PolygonElement.ElementAddress.SubMeshAddress );
-									if( EditableMesh != nullptr )
-									{
-										if( IsElementIDValid( PolygonElement, EditableMesh ) )
-										{
-											this->InsetUsingPolygonElement = PolygonElement;
-										}
-									}
-								}
-							}
-						}
-
-						if( InsetUsingPolygonElement.IsValidMeshElement() )
-						{
-							const bool bActionNeedsHoverLocation = true;
-							if( EquippedPolygonAction == EMeshEditAction::InsetPolygon )
-							{
-								StartAction( EMeshEditAction::InsetPolygon, ViewportInteractor, bActionNeedsHoverLocation, LOCTEXT( "UndoInsetPolygon", "Inset Polygon" ) );
-							}
-							else if( ensure( EquippedPolygonAction == EMeshEditAction::BevelPolygon ) )
-							{
-								StartAction( EMeshEditAction::BevelPolygon, ViewportInteractor, bActionNeedsHoverLocation, LOCTEXT( "UndoBevelPolygon", "Bevel Polygon" ) );
-							}
-
-							bOutIsInputCaptured = true;
-						}
-					}
 					else if( SelectedMeshElementType == EEditableMeshElementType::Edge && EquippedEdgeAction == EMeshEditAction::SplitEdge )
 					{
 						const bool bActionNeedsHoverLocation = true;
@@ -5281,15 +5036,32 @@ void FMeshEditorMode::OnViewportInteractionInputAction( FEditorViewportClient& V
 							UMeshEditorCommand* CommandCDO = *CommandCDOIter;
 							if( !( CommandCDO->GetClass()->GetClassFlags() & CLASS_Abstract ) )
 							{
-								if( ( CommandCDO->GetElementType() == EEditableMeshElementType::Vertex && SelectedMeshElementType == EEditableMeshElementType::Vertex && EquippedVertexAction == CommandCDO->GetCommandName() ) ||
-									( CommandCDO->GetElementType() == EEditableMeshElementType::Edge && SelectedMeshElementType == EEditableMeshElementType::Edge  && EquippedEdgeAction == CommandCDO->GetCommandName() ) ||
-									( CommandCDO->GetElementType() == EEditableMeshElementType::Polygon && SelectedMeshElementType == EEditableMeshElementType::Polygon && EquippedPolygonAction == CommandCDO->GetCommandName() ) )
+								EMeshEditAction::Type EquippedAction = EMeshEditAction::None;
+								switch( SelectedMeshElementType )
 								{
-									StartAction( EquippedEdgeAction, ViewportInteractor, CommandCDO->NeedsHoverLocation(), CommandCDO->GetUndoText() );
+									case EEditableMeshElementType::Vertex:
+										EquippedAction = EquippedVertexAction;
+										break;
+
+									case EEditableMeshElementType::Edge:
+										EquippedAction = EquippedEdgeAction;
+										break;
+
+									case EEditableMeshElementType::Polygon:
+										EquippedAction = EquippedPolygonAction;
+										break;
+								}
 
-									if( !CommandCDO->NeedsDraggingInitiated() )
+								if( CommandCDO->GetElementType() == SelectedMeshElementType && EquippedAction == CommandCDO->GetCommandName() )
+								{
+									if( CommandCDO->TryStartingToDrag( *this, ViewportInteractor ) )
 									{
-										bOutIsInputCaptured = true;
+										StartAction( EquippedAction, ViewportInteractor, CommandCDO->NeedsHoverLocation(), CommandCDO->GetUndoText() );
+
+										if( !CommandCDO->NeedsDraggingInitiated() )
+										{
+											bOutIsInputCaptured = true;
+										}
 									}
 
 									// Should always only be one candidate
@@ -6029,19 +5801,6 @@ void FMeshEditorMode::MakeVRRadialMenuActionsMenu(FMenuBuilder& MenuBuilder, TSh
 			NAME_None,
 			EUserInterfaceActionType::ToggleButton
 			);
-		MenuBuilder.AddMenuEntry(
-			LOCTEXT("Inset", "Inset"),
-			FText(),
-			FSlateIcon(FMeshEditorStyle::GetStyleSetName(), "MeshEditorMode.PolyInset"),
-			FUIAction
-			(
-				FExecuteAction::CreateLambda([this] { SetEquippedAction( EEditableMeshElementType::Polygon, EMeshEditAction::InsetPolygon ); }),
-				FCanExecuteAction::CreateSP(this, &FMeshEditorMode::IsMeshElementTypeSelectedOrIsActiveSelectionMode, EEditableMeshElementType::Polygon),
-				FIsActionChecked::CreateLambda([this] { return (EquippedPolygonAction == EMeshEditAction::InsetPolygon); })
-			),
-			NAME_None,
-			EUserInterfaceActionType::ToggleButton
-			);
 	}
 	else if (GetMeshElementSelectionMode() == EEditableMeshElementType::Edge)
 	{
@@ -6200,7 +5959,7 @@ EMeshEditAction::Type FMeshEditorMode::GetEquippedAction( const EEditableMeshEle
 			break;
 
 		case EEditableMeshElementType::Polygon:
-			EquippedAction = EquippedEdgeAction;
+			EquippedAction = EquippedPolygonAction;
 			break;
 
 		default:
@@ -6224,7 +5983,7 @@ void FMeshEditorMode::SetEquippedAction( const EEditableMeshElementType ForEleme
 			break;
 
 		case EEditableMeshElementType::Polygon:
-			EquippedEdgeAction = ActionToEquip;
+			EquippedPolygonAction = ActionToEquip;
 			break;
 
 		default:
@@ -6272,4 +6031,25 @@ void FMeshEditorMode::TrackUndo( UObject* Object, TUniquePtr<FChange> RevertChan
 }
 
 
+FMeshElement FMeshEditorMode::GetHoveredMeshElement( UViewportInteractor* ViewportInteractor ) const
+{
+	FMeshElement HoveredMeshElement;
+
+	const FMeshEditorInteractorData& InteractorData = GetMeshEditorInteractorData( ViewportInteractor );
+	if( InteractorData.HoveredMeshElement.IsValidMeshElement() )
+	{
+		const UEditableMesh* EditableMesh = FindEditableMesh( *InteractorData.HoveredMeshElement.Component, InteractorData.HoveredMeshElement.ElementAddress.SubMeshAddress );
+		if( EditableMesh != nullptr )
+		{
+			if( IsElementIDValid( InteractorData.HoveredMeshElement, EditableMesh ) )
+			{
+				HoveredMeshElement = InteractorData.HoveredMeshElement;
+			}
+		}
+	}
+
+	return HoveredMeshElement;
+}
+
+
 #undef LOCTEXT_NAMESPACE
diff --git a/Engine/Plugins/Editor/MeshEditor/Source/MeshEditor/MeshEditorMode.h b/Engine/Plugins/Editor/MeshEditor/Source/MeshEditor/MeshEditorMode.h
--- a/Engine/Plugins/Editor/MeshEditor/Source/MeshEditor/MeshEditorMode.h
+++ b/Engine/Plugins/Editor/MeshEditor/Source/MeshEditor/MeshEditorMode.h
@@ -100,7 +100,7 @@ public:
 	virtual ~FMeshEditorMode();
 
 	/** Gets an editable mesh from our cache of editable meshes for the specified sub-mesh address, or tries to create and cache a new
-	editable mesh if we haven't seen this sub-mesh address before.  Can return nullptr if no mesh was possible for that address. */
+		editable mesh if we haven't seen this sub-mesh address before.  Can return nullptr if no mesh was possible for that address. */
 	UEditableMesh* FindOrCreateEditableMesh( class UPrimitiveComponent& Component, const FEditableMeshSubMeshAddress& SubMeshAddress );
 
 
@@ -167,9 +167,15 @@ protected:
 	virtual void SetEquippedAction( const EEditableMeshElementType ForElementType, const EMeshEditAction::Type ActionToEquip ) override;
 
 	// IMeshEditorModeEditingContract interface
+	virtual const UEditableMesh* FindEditableMesh( class UPrimitiveComponent& Component, const FEditableMeshSubMeshAddress& SubMeshAddress ) const override;
 	virtual EMeshEditAction::Type GetActiveAction() const override { return ActiveAction; }
 	virtual void TrackUndo( UObject* Object, TUniquePtr<FChange> RevertChange ) override;
 	virtual void CommitSelectedMeshes() override;
+	virtual FMeshElement GetHoveredMeshElement( class UViewportInteractor* ViewportInteractor ) const override;
+	virtual bool IsMeshElementSelected( const FMeshElement MeshElement ) const override
+	{
+		return GetSelectedMeshElementIndex( MeshElement ) != INDEX_NONE;
+	}
 	virtual void GetSelectedMeshesAndVertices( TMap<UEditableMesh*, TArray<FMeshElement>>& OutMeshesAndVertices ) override 
 	{
 		GetSelectedMeshesAndElements( EEditableMeshElementType::Vertex, /* Out */ OutMeshesAndVertices ); 
@@ -324,7 +330,7 @@ protected:
 
 	/** Gets mesh editor interactor data for the specified viewport interactor.  If we've never seen this viewport interactor before,
 	    new (empty) data will be created for it on demand */
-	FMeshEditorInteractorData& GetMeshEditorInteractorData( UViewportInteractor* ViewportInteractor );
+	FMeshEditorInteractorData& GetMeshEditorInteractorData( UViewportInteractor* ViewportInteractor ) const;
 
 	/** Helper function that returns a map keying an editable mesh with its selected elements */
 	void GetSelectedMeshesAndElements( EEditableMeshElementType ElementType, TMap<UEditableMesh*, TArray<FMeshElement>>& OutMeshesAndElements );
@@ -496,7 +502,7 @@ protected:
 	double HoverFeedbackTimeValue;
 
 	/** Interactors for the mouse cursor, and also for either virtual hand (when using VR) */
-	TArray<FMeshEditorInteractorData> MeshEditorInteractorDatas;
+	mutable TArray<FMeshEditorInteractorData> MeshEditorInteractorDatas;
 
 	/** Specifies the type of element which is currently being selected */
 	EEditableMeshElementType MeshElementSelectionMode;
@@ -589,14 +595,6 @@ protected:
 	FVector ExtrudePolygonAxisDirection;
 
 
-	//
-	// InsetPolygon
-	//
-
-	/** The selected polygon we clicked on to start the inset action */
-	FMeshElement InsetUsingPolygonElement;
-
-
 	//
 	// SplitEdgeAndDragVertex
 	//
diff --git a/Engine/Plugins/Editor/MeshEditor/Source/MeshEditor/Public/IMeshEditorModeEditingContract.h b/Engine/Plugins/Editor/MeshEditor/Source/MeshEditor/Public/IMeshEditorModeEditingContract.h
--- a/Engine/Plugins/Editor/MeshEditor/Source/MeshEditor/Public/IMeshEditorModeEditingContract.h
+++ b/Engine/Plugins/Editor/MeshEditor/Source/MeshEditor/Public/IMeshEditorModeEditingContract.h
@@ -35,12 +35,6 @@ namespace EMeshEditAction
 	/** Extrude polygon by making a copy of it and allowing you to move it around freely */
 	MESHEDITOR_API extern const FName FreelyExtrudePolygon;
 
-	/** Inset polygon by replacing it with a new polygon that is bordered by polygons of a specific relative size */
-	MESHEDITOR_API extern const FName InsetPolygon;
-
-	/** Bevel polygons by adding angled bordering polygons of a specific relative size */
-	MESHEDITOR_API extern const FName BevelPolygon;
-
 	/** Extend an edge by making a copy of it and allowing you to move it around */
 	MESHEDITOR_API extern const FName ExtendEdge;
 
@@ -64,6 +58,8 @@ class IMeshEditorModeEditingContract
 
 public:
 
+	/** Gets an editable mesh from our cache of editable meshes for the specified sub-mesh address */
+	virtual const class UEditableMesh* FindEditableMesh( class UPrimitiveComponent& Component, const FEditableMeshSubMeshAddress& SubMeshAddress ) const = 0;
 
 	/** Gets the interactive action currently being performed (and previewed).  These usually happen over multiple frames, and
 	    result in a 'final' application of the change that performs a more exhaustive (and more expensive) update. */
@@ -73,6 +69,12 @@ public:
 	    currently in the middle of previewing a temporary change to meshes (bIsCapturingUndoForPreview) */
 	virtual void TrackUndo( UObject* Object, TUniquePtr<FChange> RevertChange ) = 0;
 
+	/** Returns true if the specified element is selected right now */
+	virtual bool IsMeshElementSelected( const FMeshElement MeshElement ) const = 0;
+
+	/** Returns the mesh element the specified viewport interactor is currently hovering over.  The returned element might be invalid if nothing valid is hovered right now */
+	virtual FMeshElement GetHoveredMeshElement( class UViewportInteractor* ViewportInteractor ) const = 0;
+
 	virtual void GetSelectedMeshesAndVertices( TMap<class UEditableMesh*, TArray<FMeshElement>>& OutMeshesAndVertices ) = 0;
 	virtual void GetSelectedMeshesAndEdges( TMap<class UEditableMesh*, TArray<FMeshElement>>& OutMeshesAndEdges ) = 0;
 	virtual void GetSelectedMeshesAndPolygons( TMap<class UEditableMesh*, TArray<FMeshElement>>& OutMeshesAndPolygons ) = 0;
diff --git a/Engine/Plugins/Editor/MeshEditor/Source/MeshEditor/Public/MeshEditorCommands.h b/Engine/Plugins/Editor/MeshEditor/Source/MeshEditor/Public/MeshEditorCommands.h
--- a/Engine/Plugins/Editor/MeshEditor/Source/MeshEditor/Public/MeshEditorCommands.h
+++ b/Engine/Plugins/Editor/MeshEditor/Source/MeshEditor/Public/MeshEditorCommands.h
@@ -25,8 +25,14 @@ public:
 	{
 	}
 
+	/** Called when the user starts to drag on an element.  If this returns true, then the action will begin and ApplyDuringDrag() will be called each frame until the user releases the button. */
+	virtual bool TryStartingToDrag( IMeshEditorModeEditingContract& MeshEditorMode, class UViewportInteractor* ViewportInteractor )
+	{
+		return true;
+	}
+
 	/** Applies this command every frame while dragging */
-	virtual void ApplyDuringDrag( class IMeshEditorModeEditingContract& MeshEditorMode, bool& bOutShouldDeselectAllFirst, TArray<FMeshElement>& OutMeshElementsToSelect )
+	virtual void ApplyDuringDrag( class IMeshEditorModeEditingContract& MeshEditorMode, class UViewportInteractor* ViewportInteractor, bool& bOutShouldDeselectAllFirst, TArray<FMeshElement>& OutMeshElementsToSelect )
 	{
 	}
 
@@ -287,12 +293,6 @@ public:
 	/** Sets the primary action to freely extrude polygons */
 	TSharedPtr<FUICommandInfo> FreelyExtrudePolygon;
 
-	/** Sets the primary action to inset polygons */
-	TSharedPtr<FUICommandInfo> InsetPolygon;
-
-	/** Sets the primary action to bevel polygons */
-	TSharedPtr<FUICommandInfo> BevelPolygon;
-
 	/** Flips the currently selected polygon(s) */
 	TSharedPtr<FUICommandInfo> FlipPolygon;
 
diff --git a/Engine/Plugins/Editor/MeshEditor/Source/PolygonModeling/InsertEdgeLoop.cpp b/Engine/Plugins/Editor/MeshEditor/Source/PolygonModeling/InsertEdgeLoop.cpp
--- a/Engine/Plugins/Editor/MeshEditor/Source/PolygonModeling/InsertEdgeLoop.cpp
+++ b/Engine/Plugins/Editor/MeshEditor/Source/PolygonModeling/InsertEdgeLoop.cpp
@@ -19,7 +19,7 @@ void UInsertEdgeLoopCommand::RegisterUICommand( FBindingContext* BindingContext
 }
 
 
-void UInsertEdgeLoopCommand::ApplyDuringDrag( IMeshEditorModeEditingContract& MeshEditorMode, bool& bOutShouldDeselectAllFirst, TArray<FMeshElement>& OutMeshElementsToSelect )
+void UInsertEdgeLoopCommand::ApplyDuringDrag( IMeshEditorModeEditingContract& MeshEditorMode, UViewportInteractor* ViewportInteractor, bool& bOutShouldDeselectAllFirst, TArray<FMeshElement>& OutMeshElementsToSelect )
 {
 	// Insert edge loop
 	static TMap< UEditableMesh*, TArray< FMeshElement > > SelectedMeshesAndEdges;
@@ -101,7 +101,7 @@ void UInsertEdgeLoopCommand::AddToVRRadialMenuActionsMenu( IMeshEditorModeUICont
 	if( MeshEditorMode.GetMeshElementSelectionMode() == EEditableMeshElementType::Edge )
 	{
 		MenuBuilder.AddMenuEntry(
-			LOCTEXT("InsertEdgeLoop", "Insert Loop"),
+			LOCTEXT("VRInsertEdgeLoop", "Insert Loop"),
 			FText(),
 			FSlateIcon(TEMPHACK_StyleSetName, "MeshEditorMode.EdgeInsert"),	// @todo mesheditor extensibility: TEMPHACK for style; Need PolygonModelingStyle, probably.  Or we're just cool with exporting MeshEditorModeStyle, since we're all the same plugin after all.
 			MakeUIAction( MeshEditorMode ),
diff --git a/Engine/Plugins/Editor/MeshEditor/Source/PolygonModeling/InsertEdgeLoop.h b/Engine/Plugins/Editor/MeshEditor/Source/PolygonModeling/InsertEdgeLoop.h
--- a/Engine/Plugins/Editor/MeshEditor/Source/PolygonModeling/InsertEdgeLoop.h
+++ b/Engine/Plugins/Editor/MeshEditor/Source/PolygonModeling/InsertEdgeLoop.h
@@ -6,7 +6,7 @@
 #include "InsertEdgeLoop.generated.h"
 
 
-/** With an edge select, inserts a loop of edge perpendicular to that edge while dragging */
+/** With an edge selected, inserts a loop of edge perpendicular to that edge while dragging */
 UCLASS()
 class UInsertEdgeLoopCommand : public UMeshEditorEdgeCommand
 {
@@ -24,7 +24,7 @@ protected:
 
 	// Overrides
 	virtual void RegisterUICommand( class FBindingContext* BindingContext ) override;
-	virtual void ApplyDuringDrag( IMeshEditorModeEditingContract& MeshEditorMode, bool& bOutShouldDeselectAllFirst, TArray<FMeshElement>& OutMeshElementsToSelect ) override;
+	virtual void ApplyDuringDrag( IMeshEditorModeEditingContract& MeshEditorMode, class UViewportInteractor* ViewportInteractor, bool& bOutShouldDeselectAllFirst, TArray<FMeshElement>& OutMeshElementsToSelect ) override;
 	virtual void AddToVRRadialMenuActionsMenu( class IMeshEditorModeUIContract& MeshEditorMode, class FMenuBuilder& MenuBuilder, TSharedPtr<FUICommandList> CommandList, const FName TEMPHACK_StyleSetName, class UVREditorMode* VRMode ) override;
 
 };
diff --git a/Engine/Plugins/Editor/MeshEditor/Source/PolygonModeling/PolygonModeling.Build.cs b/Engine/Plugins/Editor/MeshEditor/Source/PolygonModeling/PolygonModeling.Build.cs
--- a/Engine/Plugins/Editor/MeshEditor/Source/PolygonModeling/PolygonModeling.Build.cs
+++ b/Engine/Plugins/Editor/MeshEditor/Source/PolygonModeling/PolygonModeling.Build.cs
@@ -16,7 +16,8 @@ namespace UnrealBuildTool.Rules
 					"UnrealEd",
 					"MeshEditingRuntime",
 					"MeshEditor",
-					"SlateCore"
+					"SlateCore",
+					"ViewportInteraction"
 				}
 			);
 		}
diff --git a/Engine/Plugins/Editor/MeshEditor/Source/PolygonModeling/TessellatePolygon.cpp b/Engine/Plugins/Editor/MeshEditor/Source/PolygonModeling/TessellatePolygon.cpp
new file mode 100644
--- /dev/null
+++ b/Engine/Plugins/Editor/MeshEditor/Source/PolygonModeling/TessellatePolygon.cpp
@@ -0,0 +1,98 @@
+// Copyright 1998-2016 Epic Games, Inc. All Rights Reserved.
+
+#include "TessellatePolygon.h"
+#include "IMeshEditorModeEditingContract.h"
+#include "UICommandInfo.h"
+#include "EditableMesh.h"
+#include "MeshElement.h"
+
+
+#define LOCTEXT_NAMESPACE "MeshEditorMode"
+
+
+void UTessellatePolygonCommand::RegisterUICommand( FBindingContext* BindingContext )
+{
+	UI_COMMAND_EXT( BindingContext, /* Out */ UICommandInfo, "TessellatePolygon", "Tessellate Polygon", "Tessellate selected polygons into smaller polygons.", EUserInterfaceActionType::Button, FInputChord() );
+}
+
+
+void UTessellatePolygonCommand::Execute( IMeshEditorModeEditingContract& MeshEditorMode )
+{
+	if( MeshEditorMode.GetActiveAction() != EMeshEditAction::None )
+	{
+		return;
+	}
+
+	static TMap<UEditableMesh*, TArray<FMeshElement>> SelectedMeshesAndPolygons;
+	MeshEditorMode.GetSelectedMeshesAndPolygons( SelectedMeshesAndPolygons );
+	if( SelectedMeshesAndPolygons.Num() == 0 )
+	{
+		return;
+	}
+
+	FScopedTransaction Transaction( LOCTEXT( "UndoTessellatePolygon", "Tessellate Polygon" ) );
+
+	MeshEditorMode.CommitSelectedMeshes();
+
+	// Refresh selection (committing may have created a new mesh instance)
+	MeshEditorMode.GetSelectedMeshesAndPolygons( SelectedMeshesAndPolygons );
+
+	// Deselect the mesh elements before we delete them.  This will make sure they become selected again after undo.
+	MeshEditorMode.DeselectMeshElements( SelectedMeshesAndPolygons );
+
+	static TArray<FMeshElement> MeshElementsToSelect;
+	MeshElementsToSelect.Reset();
+
+	for( const auto& SelectedMeshAndPolygons : SelectedMeshesAndPolygons )
+	{
+		UEditableMesh* EditableMesh = SelectedMeshAndPolygons.Key;
+		const TArray<FMeshElement>& PolygonElements = SelectedMeshAndPolygons.Value;
+
+		static TArray<FPolygonRef> PolygonsToTessellate;
+		PolygonsToTessellate.Reset( PolygonElements.Num() );
+
+		for( const auto& PolygonElement : PolygonElements )
+		{
+			const FPolygonRef PolygonRef( PolygonElement.ElementAddress.SectionID, FPolygonID( PolygonElement.ElementAddress.ElementID ) );
+			PolygonsToTessellate.Add( PolygonRef );
+		}
+
+
+		EditableMesh->StartModification( EMeshModificationType::Final, EMeshTopologyChange::TopologyChange );
+
+		// @todo mesheditor: Expose as configurable parameter
+		const ETriangleTessellationMode TriangleTessellationMode = ETriangleTessellationMode::FourTriangles;
+
+		static TArray<FPolygonRef> NewPolygonRefs;
+		EditableMesh->TessellatePolygons( PolygonsToTessellate, TriangleTessellationMode, /* Out */ NewPolygonRefs );
+
+		// Select the new polygons
+		for( const FPolygonRef NewPolygonRef : NewPolygonRefs )
+		{
+			FMeshElement NewPolygonMeshElement;
+			{
+				NewPolygonMeshElement.Component = PolygonElements[0].Component;
+				NewPolygonMeshElement.ElementAddress = PolygonElements[0].ElementAddress;
+				NewPolygonMeshElement.ElementAddress.ElementType = EEditableMeshElementType::Polygon;
+				NewPolygonMeshElement.ElementAddress.SectionID = NewPolygonRef.SectionID;
+				NewPolygonMeshElement.ElementAddress.ElementID = NewPolygonRef.PolygonID;
+			}
+
+			MeshElementsToSelect.Add( NewPolygonMeshElement );
+		}
+
+		EditableMesh->EndModification();
+
+		MeshEditorMode.TrackUndo( EditableMesh, EditableMesh->MakeUndo() );
+	}
+
+	// Make sure we're not still hovering over a polygon we removed
+	MeshEditorMode.ClearInvalidSelectedElements();
+
+	// Select the new smaller polygons
+	MeshEditorMode.SelectMeshElements( MeshElementsToSelect );
+}
+
+
+#undef LOCTEXT_NAMESPACE
+
diff --git a/Engine/Plugins/Editor/MeshEditor/Source/PolygonModeling/TessellatePolygon.h b/Engine/Plugins/Editor/MeshEditor/Source/PolygonModeling/TessellatePolygon.h
new file mode 100644
--- /dev/null
+++ b/Engine/Plugins/Editor/MeshEditor/Source/PolygonModeling/TessellatePolygon.h
@@ -0,0 +1,21 @@
+// Copyright 1998-2016 Epic Games, Inc. All Rights Reserved.
+
+#pragma once
+
+#include "MeshEditorCommands.h"
+#include "TessellatePolygon.generated.h"
+
+
+/** Tessellates selected polygons into smaller polygons */
+UCLASS()
+class UTessellatePolygonCommand : public UMeshEditorPolygonCommand
+{
+	GENERATED_BODY()
+
+protected:
+
+	// Overrides
+	virtual void RegisterUICommand( class FBindingContext* BindingContext ) override;
+	virtual void Execute( class IMeshEditorModeEditingContract& MeshEditorMode ) override;
+
+};
diff --git a/Engine/Plugins/Editor/MeshEditor/Source/PolygonModeling/TessellatePolygons.cpp b/Engine/Plugins/Editor/MeshEditor/Source/PolygonModeling/TessellatePolygons.cpp
deleted file mode 100644
diff --git a/Engine/Plugins/Editor/MeshEditor/Source/PolygonModeling/TessellatePolygons.h b/Engine/Plugins/Editor/MeshEditor/Source/PolygonModeling/TessellatePolygons.h
deleted file mode 100644
