Mesh Editor extensibility (phase 3) (part 2)
- Added missing source files

diff --git a/Engine/Plugins/Editor/MeshEditor/Source/PolygonModeling/BevelOrInsetPolygon.cpp b/Engine/Plugins/Editor/MeshEditor/Source/PolygonModeling/BevelOrInsetPolygon.cpp
new file mode 100644
--- /dev/null
+++ b/Engine/Plugins/Editor/MeshEditor/Source/PolygonModeling/BevelOrInsetPolygon.cpp
@@ -0,0 +1,325 @@
+// Copyright 1998-2016 Epic Games, Inc. All Rights Reserved.
+
+#include "BevelOrInsetPolygon.h"
+#include "IMeshEditorModeEditingContract.h"
+#include "IMeshEditorModeUIContract.h"
+#include "UICommandInfo.h"
+#include "EditableMesh.h"
+#include "MeshElement.h"
+#include "MultiBoxBuilder.h"
+#include "UICommandList.h"
+#include "ViewportInteractor.h"
+
+#define LOCTEXT_NAMESPACE "MeshEditorMode"
+
+
+namespace BevelOrInsetPolygonHelpers
+{
+	/** The selected polygon we clicked on to start the inset action */
+	FMeshElement InsetUsingPolygonElement;
+
+	enum class EBevelOrInset
+	{
+		Bevel,
+		Inset
+	};
+
+
+	// Figures out how far we should inset on a polygon by figuring out the maximum distance from all of the polygon's perimeter
+	// vertices to the center of the polygon, then subtracting the distance between the hovered impact point and polygon center.
+	static void FindInsetAmount(
+		IMeshEditorModeEditingContract& MeshEditorMode,
+		UViewportInteractor* ViewportInteractor,
+		const FPolygonRef& PolygonRef,
+		UPrimitiveComponent& Component,
+		const UEditableMesh* EditableMesh,
+		float& OutInsetFixedDistance,
+		float& OutInsetProgressTowardCenter )
+	{
+		// @todo grabber: Glitches out when using grabber sphere near the corner of an inset polygon.
+		check( ViewportInteractor != nullptr );
+
+		OutInsetFixedDistance = 0.0f;
+		OutInsetProgressTowardCenter = 0.0f;
+
+		// @todo mesheditor extensibility: Need to decide whether to expose FMeshEditorInteractorData (currently we go straight to the ViewportInteractor for LaserStart, LasertEnd, bLaserIsValid)
+		// @todo mesheditor grabber: Needs grabber sphere support
+		FVector LaserStart, LaserEnd;
+		const bool bLaserIsValid = ViewportInteractor->GetLaserPointer( /* Out */ LaserStart, /* Out */ LaserEnd );
+		if( bLaserIsValid )
+		{
+			const FMatrix ComponentToWorldMatrix = Component.GetRenderMatrix();
+
+			const FPlane PolygonPlane = EditableMesh->ComputePolygonPlane( PolygonRef );
+			const FVector PolygonCenter = EditableMesh->ComputePolygonCenter( PolygonRef );
+
+			const FVector ComponentSpaceRayStart = ComponentToWorldMatrix.InverseTransformPosition( LaserStart );
+			const FVector ComponentSpaceRayEnd = ComponentToWorldMatrix.InverseTransformPosition( LaserEnd );
+
+			FVector ComponentSpaceRayIntersectionWithPolygonPlane;
+			if( FMath::SegmentPlaneIntersection( ComponentSpaceRayStart, ComponentSpaceRayEnd, PolygonPlane, /* Out */ ComponentSpaceRayIntersectionWithPolygonPlane ) )
+			{
+				// @todo mesheditor debug
+				// DrawDebugSphere( This->GetWorld(), ComponentToWorldMatrix.TransformPosition( ComponentSpaceRayIntersectionWithPolygonPlane ), 1.5f, 32, FColor::Red, false, 0.0f );
+
+				bool bFoundTriangle = false;
+				float CenterWeight = 0.0f;
+				FVector BestEdgeVertex0Position, BestEdgeVertex1Position;
+				{
+					static TArray<FVertexID> PerimeterVertexIDs;
+					EditableMesh->GetPolygonPerimeterVertices( PolygonRef, /* Out */ PerimeterVertexIDs );
+
+					for( int32 VertexNumber = 0; VertexNumber < PerimeterVertexIDs.Num(); ++VertexNumber )
+					{
+						const int32 NextVertexNumber = ( VertexNumber + 1 ) % PerimeterVertexIDs.Num();
+
+
+						const FVertexID EdgeVertex0 = PerimeterVertexIDs[ VertexNumber ];
+						const FVertexID EdgeVertex1 = PerimeterVertexIDs[ NextVertexNumber ];
+
+						const FVector EdgeVertex0Position = EditableMesh->GetVertexAttribute( EdgeVertex0, UEditableMeshAttribute::VertexPosition(), 0 );
+						const FVector EdgeVertex1Position = EditableMesh->GetVertexAttribute( EdgeVertex1, UEditableMeshAttribute::VertexPosition(), 0 );
+
+						BestEdgeVertex0Position = EdgeVertex0Position;
+						BestEdgeVertex1Position = EdgeVertex1Position;
+
+						const FVector VertexWeights = FMath::ComputeBaryCentric2D( ComponentSpaceRayIntersectionWithPolygonPlane, EdgeVertex0Position, EdgeVertex1Position, PolygonCenter );
+						if( VertexWeights.X >= 0.0f && VertexWeights.Y >= 0.0f && VertexWeights.Z >= 0.0f )
+						{
+							CenterWeight = VertexWeights.Z;
+
+							bFoundTriangle = true;
+							break;
+						}
+					}
+				}
+
+				if( bFoundTriangle )
+				{
+					// @todo mesheditor debug
+					// DrawDebugLine( This->GetWorld(), ComponentToWorldMatrix.TransformPosition( BestEdgeVertex0Position ), ComponentToWorldMatrix.TransformPosition( BestEdgeVertex1Position ), FColor::Yellow, false, 0.0f, 0, 2.0f );
+
+					OutInsetProgressTowardCenter = CenterWeight;
+				}
+			}
+		}
+	}
+
+
+	static bool TryStartingToDrag( IMeshEditorModeEditingContract& MeshEditorMode, UViewportInteractor* ViewportInteractor )
+	{
+		bool bCanStartDragging = false;
+
+		// Figure out which polygon to inset
+		InsetUsingPolygonElement = FMeshElement();
+		{
+			const FMeshElement& PolygonElement = MeshEditorMode.GetHoveredMeshElement( ViewportInteractor );
+			if( PolygonElement.IsValidMeshElement() &&
+				PolygonElement.ElementAddress.ElementType == EEditableMeshElementType::Polygon &&
+				MeshEditorMode.IsMeshElementSelected( PolygonElement ) )
+			{
+				InsetUsingPolygonElement = PolygonElement;
+				bCanStartDragging = true;
+			}
+		}
+
+		return bCanStartDragging;
+	}
+
+
+	static void ApplyDuringDrag( const EBevelOrInset BevelOrInset, IMeshEditorModeEditingContract& MeshEditorMode, class UViewportInteractor* ViewportInteractor, bool& bOutShouldDeselectAllFirst, TArray<FMeshElement>& OutMeshElementsToSelect )
+	{
+		static TMap< UEditableMesh*, TArray< FMeshElement > > MeshesWithPolygonsToInset;
+		MeshEditorMode.GetSelectedMeshesAndPolygons( /* Out */ MeshesWithPolygonsToInset );
+
+		if( MeshesWithPolygonsToInset.Num() > 0 )
+		{
+			if( InsetUsingPolygonElement.IsValidMeshElement() &&
+				InsetUsingPolygonElement.ElementAddress.ElementType == EEditableMeshElementType::Polygon &&
+				MeshEditorMode.IsMeshElementSelected( InsetUsingPolygonElement ) )
+			{
+				UPrimitiveComponent* InsetUsingComponent = InsetUsingPolygonElement.Component.Get();
+				check( InsetUsingComponent != nullptr );
+
+				const UEditableMesh* InsetUsingEditableMesh = MeshEditorMode.FindEditableMesh( *InsetUsingComponent, InsetUsingPolygonElement.ElementAddress.SubMeshAddress );
+				if( ensure( InsetUsingEditableMesh != nullptr ) )
+				{
+					const FPolygonRef InsetUsingPolygonRef( InsetUsingPolygonElement.ElementAddress.SectionID, FPolygonID( InsetUsingPolygonElement.ElementAddress.ElementID ) );
+
+					// Figure out how far to inset the polygon
+					float InsetFixedDistance = 0.0f;
+					float InsetProgressTowardCenter = 0.0f;
+					FindInsetAmount(
+						MeshEditorMode,
+						ViewportInteractor,
+						InsetUsingPolygonRef,
+						*InsetUsingComponent,
+						InsetUsingEditableMesh,
+						/* Out */ InsetFixedDistance,
+						/* Out */ InsetProgressTowardCenter );
+
+					if( InsetFixedDistance > SMALL_NUMBER ||
+						InsetProgressTowardCenter > SMALL_NUMBER )
+					{
+						// Deselect the mesh elements before we delete them.  This will make sure they become selected again after undo.
+						MeshEditorMode.DeselectMeshElements( MeshesWithPolygonsToInset );
+
+						for( auto& MeshAndPolygons : MeshesWithPolygonsToInset )
+						{
+							UEditableMesh* EditableMesh = MeshAndPolygons.Key;
+							const TArray<FMeshElement>& PolygonsToInset = MeshAndPolygons.Value;
+
+							UPrimitiveComponent* Component = PolygonsToInset[ 0 ].Component.Get();	// NOTE: All polygons in this array belong to the same mesh/component, so we just need the first element
+							check( Component != nullptr );
+
+
+							static TArray<FPolygonRef> PolygonRefsToInset;
+							PolygonRefsToInset.Reset();
+							for( const FMeshElement& PolygonToInset : PolygonsToInset )
+							{
+								FPolygonRef PolygonRef( PolygonToInset.ElementAddress.SectionID, FPolygonID( PolygonToInset.ElementAddress.ElementID ) );
+								PolygonRefsToInset.Add( PolygonRef );
+							}
+
+							{
+								verify( !EditableMesh->AnyChangesToUndo() );
+
+								// Inset time!!
+								static TArray<FPolygonRef> NewCenterInsetPolygons;
+								static TArray<FPolygonRef> NewSideInsetPolygons;
+								if( BevelOrInset == EBevelOrInset::Inset )
+								{
+									const EInsetPolygonsMode InsetPolygonsMode = EInsetPolygonsMode::All;	// @todo mesheditor inset: Make configurable?
+																											// @todo mesheditor inset: Add options for Fixed distance (instead of Percentage distance, like now.)
+
+									EditableMesh->InsetPolygons( PolygonRefsToInset, InsetFixedDistance, InsetProgressTowardCenter, InsetPolygonsMode, /* Out */ NewCenterInsetPolygons, /* Out */ NewSideInsetPolygons );
+								}
+								else if( ensure( BevelOrInset == EBevelOrInset::Bevel ) )
+								{
+									EditableMesh->BevelPolygons( PolygonRefsToInset, InsetFixedDistance, InsetProgressTowardCenter, /* Out */ NewCenterInsetPolygons, /* Out */ NewSideInsetPolygons );
+								}
+
+								// Make sure the new polygons are selected.  The old polygon was deleted and will become deselected automatically.
+								if( NewCenterInsetPolygons.Num() > 0 )
+								{
+									for( int32 NewPolygonNumber = 0; NewPolygonNumber < NewCenterInsetPolygons.Num(); ++NewPolygonNumber )
+									{
+										const FPolygonRef& NewInsetCenterPolygon = NewCenterInsetPolygons[ NewPolygonNumber ];
+
+										const FMeshElement& MeshElement = PolygonsToInset[ NewPolygonNumber ];
+
+										FMeshElement PolygonMeshElement;
+										{
+											PolygonMeshElement.Component = MeshElement.Component;
+											PolygonMeshElement.ElementAddress = MeshElement.ElementAddress;
+											PolygonMeshElement.ElementAddress.SectionID = NewInsetCenterPolygon.SectionID;
+											PolygonMeshElement.ElementAddress.ElementID = NewInsetCenterPolygon.PolygonID;
+										}
+
+										// Queue selection of this new element.  We don't want it to be part of the current action.
+										OutMeshElementsToSelect.Add( PolygonMeshElement );
+									}
+								}
+								else if( ensure( NewSideInsetPolygons.Num() > 0 ) )
+								{
+									for( int32 NewPolygonNumber = 0; NewPolygonNumber < NewSideInsetPolygons.Num(); ++NewPolygonNumber )
+									{
+										const FPolygonRef& NewInsetSidePolygon = NewSideInsetPolygons[ NewPolygonNumber ];
+
+										const FMeshElement& MeshElement = PolygonsToInset[ NewPolygonNumber ];
+
+										FMeshElement PolygonMeshElement;
+										{
+											PolygonMeshElement.Component = MeshElement.Component;
+											PolygonMeshElement.ElementAddress = MeshElement.ElementAddress;
+											PolygonMeshElement.ElementAddress.SectionID = NewInsetSidePolygon.SectionID;
+											PolygonMeshElement.ElementAddress.ElementID = NewInsetSidePolygon.PolygonID;
+										}
+
+										// Queue selection of this new element.  We don't want it to be part of the current action.
+										OutMeshElementsToSelect.Add( PolygonMeshElement );
+									}
+								}
+							}
+
+							MeshEditorMode.TrackUndo( EditableMesh, EditableMesh->MakeUndo() );
+						}
+					}
+				}
+			}
+		}
+	}
+}
+
+
+void UBevelPolygonCommand::RegisterUICommand( FBindingContext* BindingContext )
+{
+	UI_COMMAND_EXT( BindingContext, /* Out */ UICommandInfo, "BevelPolygon", "Bevel Polygon Mode", "Set the primary action to bevel polygons.", EUserInterfaceActionType::RadioButton, FInputChord() );
+}
+
+
+bool UBevelPolygonCommand::TryStartingToDrag( IMeshEditorModeEditingContract& MeshEditorMode, UViewportInteractor* ViewportInteractor )
+{
+	return BevelOrInsetPolygonHelpers::TryStartingToDrag( MeshEditorMode, ViewportInteractor );
+}
+
+
+void UBevelPolygonCommand::ApplyDuringDrag( IMeshEditorModeEditingContract& MeshEditorMode, UViewportInteractor* ViewportInteractor, bool& bOutShouldDeselectAllFirst, TArray<FMeshElement>& OutMeshElementsToSelect )
+{
+	BevelOrInsetPolygonHelpers::ApplyDuringDrag( BevelOrInsetPolygonHelpers::EBevelOrInset::Bevel, MeshEditorMode, ViewportInteractor, bOutShouldDeselectAllFirst, OutMeshElementsToSelect );
+}
+
+
+void UBevelPolygonCommand::AddToVRRadialMenuActionsMenu( IMeshEditorModeUIContract& MeshEditorMode, FMenuBuilder& MenuBuilder, TSharedPtr<FUICommandList> CommandList, const FName TEMPHACK_StyleSetName, class UVREditorMode* VRMode )
+{
+	if( MeshEditorMode.GetMeshElementSelectionMode() == EEditableMeshElementType::Polygon )
+	{
+		MenuBuilder.AddMenuEntry(
+			LOCTEXT( "VRBevelPolygon", "Bevel" ),
+			FText(),
+			FSlateIcon( TEMPHACK_StyleSetName, "MeshEditorMode.PolyBevel" ),
+			MakeUIAction( MeshEditorMode ),
+			NAME_None,
+			EUserInterfaceActionType::ToggleButton
+		);
+	}
+}
+
+
+
+void UInsetPolygonCommand::RegisterUICommand( FBindingContext* BindingContext )
+{
+	UI_COMMAND_EXT( BindingContext, /* Out */ UICommandInfo, "InsetPolygon", "Inset Polygon Mode", "Set the primary action to inset polygons.", EUserInterfaceActionType::RadioButton, FInputChord() );
+}
+
+
+bool UInsetPolygonCommand::TryStartingToDrag( IMeshEditorModeEditingContract& MeshEditorMode, UViewportInteractor* ViewportInteractor )
+{
+	return BevelOrInsetPolygonHelpers::TryStartingToDrag( MeshEditorMode, ViewportInteractor );
+}
+
+
+void UInsetPolygonCommand::ApplyDuringDrag( IMeshEditorModeEditingContract& MeshEditorMode, UViewportInteractor* ViewportInteractor, bool& bOutShouldDeselectAllFirst, TArray<FMeshElement>& OutMeshElementsToSelect )
+{
+	BevelOrInsetPolygonHelpers::ApplyDuringDrag( BevelOrInsetPolygonHelpers::EBevelOrInset::Inset, MeshEditorMode, ViewportInteractor, bOutShouldDeselectAllFirst, OutMeshElementsToSelect );
+}
+
+
+void UInsetPolygonCommand::AddToVRRadialMenuActionsMenu( IMeshEditorModeUIContract& MeshEditorMode, FMenuBuilder& MenuBuilder, TSharedPtr<FUICommandList> CommandList, const FName TEMPHACK_StyleSetName, class UVREditorMode* VRMode )
+{
+	if( MeshEditorMode.GetMeshElementSelectionMode() == EEditableMeshElementType::Polygon )
+	{
+		MenuBuilder.AddMenuEntry(
+			LOCTEXT( "VRInsetPolygon", "Inset" ),
+			FText(),
+			FSlateIcon( TEMPHACK_StyleSetName, "MeshEditorMode.PolyInsert" ),
+			MakeUIAction( MeshEditorMode ),
+			NAME_None,
+			EUserInterfaceActionType::ToggleButton
+		);
+	}
+}
+
+
+#undef LOCTEXT_NAMESPACE
+
diff --git a/Engine/Plugins/Editor/MeshEditor/Source/PolygonModeling/BevelOrInsetPolygon.h b/Engine/Plugins/Editor/MeshEditor/Source/PolygonModeling/BevelOrInsetPolygon.h
new file mode 100644
--- /dev/null
+++ b/Engine/Plugins/Editor/MeshEditor/Source/PolygonModeling/BevelOrInsetPolygon.h
@@ -0,0 +1,56 @@
+// Copyright 1998-2016 Epic Games, Inc. All Rights Reserved.
+
+#pragma once
+
+#include "MeshEditorCommands.h"
+#include "BevelOrInsetPolygon.generated.h"
+
+
+/** Adds a beveled edge to an existing polygon */
+UCLASS()
+class UBevelPolygonCommand : public UMeshEditorPolygonCommand
+{
+	GENERATED_BODY()
+
+protected:
+
+	UBevelPolygonCommand()
+	{
+		bIsMode = true;
+		UndoText = NSLOCTEXT( "MeshEditor", "UndoBevelPolygon", "Bevel Polygon" );
+		bNeedsHoverLocation = true;
+		bNeedsDraggingInitiated = false;
+	}
+
+	// Overrides
+	virtual void RegisterUICommand( class FBindingContext* BindingContext ) override;
+	virtual bool TryStartingToDrag( IMeshEditorModeEditingContract& MeshEditorMode, class UViewportInteractor* ViewportInteractor ) override;
+	virtual void ApplyDuringDrag( IMeshEditorModeEditingContract& MeshEditorMode, class UViewportInteractor* ViewportInteractor, bool& bOutShouldDeselectAllFirst, TArray<FMeshElement>& OutMeshElementsToSelect ) override;
+	virtual void AddToVRRadialMenuActionsMenu( class IMeshEditorModeUIContract& MeshEditorMode, class FMenuBuilder& MenuBuilder, TSharedPtr<FUICommandList> CommandList, const FName TEMPHACK_StyleSetName, class UVREditorMode* VRMode ) override;
+
+};
+
+
+/** Adds a new polygon on the inside of an existing polygon, allowing the user to drag to set exactly where it should be placed. */
+UCLASS()
+class UInsetPolygonCommand : public UMeshEditorPolygonCommand
+{
+	GENERATED_BODY()
+
+protected:
+
+	UInsetPolygonCommand()
+	{
+		bIsMode = true;
+		UndoText = NSLOCTEXT( "MeshEditor", "UndoInsetPolygon", "Inset Polygon" );
+		bNeedsHoverLocation = true;
+		bNeedsDraggingInitiated = false;
+	}
+
+	// Overrides
+	virtual void RegisterUICommand( class FBindingContext* BindingContext ) override;
+	virtual bool TryStartingToDrag( class IMeshEditorModeEditingContract& MeshEditorMode, class UViewportInteractor* ViewportInteractor ) override;
+	virtual void ApplyDuringDrag( IMeshEditorModeEditingContract& MeshEditorMode, class UViewportInteractor* ViewportInteractor, bool& bOutShouldDeselectAllFirst, TArray<FMeshElement>& OutMeshElementsToSelect ) override;
+	virtual void AddToVRRadialMenuActionsMenu( class IMeshEditorModeUIContract& MeshEditorMode, class FMenuBuilder& MenuBuilder, TSharedPtr<FUICommandList> CommandList, const FName TEMPHACK_StyleSetName, class UVREditorMode* VRMode ) override;
+
+};
