Mesh Editor extensibility (phase 4)
- Extrude Polygon has been migrated to a new module
- Removed code for 'freely extruding polygons'.  Wasn't useful in it's current implementation.

diff --git a/Engine/Plugins/Editor/MeshEditor/Source/MeshEditor/MeshEditorCommands.cpp b/Engine/Plugins/Editor/MeshEditor/Source/MeshEditor/MeshEditorCommands.cpp
--- a/Engine/Plugins/Editor/MeshEditor/Source/MeshEditor/MeshEditorCommands.cpp
+++ b/Engine/Plugins/Editor/MeshEditor/Source/MeshEditor/MeshEditorCommands.cpp
@@ -144,9 +144,6 @@ FMeshEditorPolygonCommands::FMeshEditorPolygonCommands()
 void FMeshEditorPolygonCommands::RegisterCommands()
 {
 	UI_COMMAND(MovePolygon, "Move Polygon Mode", "Set the primary action to move polygons.", EUserInterfaceActionType::RadioButton, FInputChord(EKeys::F1));
-	UI_COMMAND(ExtrudePolygon, "Extrude Polygon Mode", "Set the primary action to extrude polygons.", EUserInterfaceActionType::RadioButton, FInputChord(EKeys::F2));
-	// UI_COMMAND(FreelyExtrudePolygon, "Freely Extrude Polygon Mode", "Set the primary action to freely extrude polygons.", EUserInterfaceActionType::RadioButton, FInputChord(EKeys::F3));
-
 	UI_COMMAND(FlipPolygon, "Flip Polygon", "Flip the currently selected polygons.", EUserInterfaceActionType::Button, FInputChord(EKeys::F, EModifierKey::Shift));
 	UI_COMMAND(TriangulatePolygon, "Triangulate Polygon", "Triangulate the currently selected polygons.", EUserInterfaceActionType::Button, FInputChord(EKeys::T));
 	UI_COMMAND(AssignMaterial, "Assign Material", "Assigns the highlighted material in the Content Browser to the currently selected polygons.", EUserInterfaceActionType::Button, FInputChord(EKeys::M));
diff --git a/Engine/Plugins/Editor/MeshEditor/Source/MeshEditor/MeshEditorMode.cpp b/Engine/Plugins/Editor/MeshEditor/Source/MeshEditor/MeshEditorMode.cpp
--- a/Engine/Plugins/Editor/MeshEditor/Source/MeshEditor/MeshEditorMode.cpp
+++ b/Engine/Plugins/Editor/MeshEditor/Source/MeshEditor/MeshEditorMode.cpp
@@ -49,8 +49,6 @@ namespace EMeshEditAction
 	const Type Move( "Move" );
 	const Type SplitEdge( "SplitEdge" );
 	const Type SplitEdgeAndDragVertex( "SplitEdgeAndDragVertex" );
-	const Type ExtrudePolygon( "ExtrudePolygon" );
-	const Type FreelyExtrudePolygon( "FreelyExtrudePolygon" );
 	const Type ExtendEdge( "ExtendEdge" );
 	const Type ExtendVertex( "ExtendVertex" );
 	const Type EditVertexCornerSharpness( "EditVertexCornerSharpness" );
@@ -394,8 +392,6 @@ FMeshEditorMode::FMeshEditorMode()
 	  ActiveActionInteractor( nullptr ),
 	  bActiveActionNeedsHoverLocation( false ),
 	  bIsFirstActiveActionUpdate( false ),
-	  ExtrudePolygonAxisOrigin( FVector::ZeroVector ),
-	  ExtrudePolygonAxisDirection( FVector::ZeroVector ),
 	  SplitEdgeMeshesAndEdgesToSplit(),
 	  SplitEdgeSplitList(),
 	  EditSharpnessStartLocation( FVector::ZeroVector ),
@@ -541,9 +537,6 @@ void FMeshEditorMode::BindCommands()
 	RegisterEdgeEditingMode( MeshEditorEdgeActions.EditEdgeCreaseSharpness, EMeshEditAction::EditEdgeCreaseSharpness );
 
 	RegisterPolygonEditingMode( MeshEditorPolygonActions.MovePolygon, EMeshEditAction::Move );
-	RegisterPolygonEditingMode( MeshEditorPolygonActions.ExtrudePolygon, EMeshEditAction::ExtrudePolygon );
-//	RegisterPolygonEditingMode( MeshEditorPolygonActions.FreelyExtrudePolygon, EMeshEditAction::FreelyExtrudePolygon );	// @todo mesheditor: "Freely Extrude" disabled until we have proper gizmo/snap support
-
 	RegisterCommonEditingMode( MeshEditorCommonActions.DrawVertices, EMeshEditAction::DrawVertices );
 
 	// Register commands which work regardless of which element type is selected
@@ -3200,60 +3193,6 @@ void FMeshEditorMode::UpdateActiveAction( const bool bIsActionFinishing )
 	// sure that something is always selected after every type of interactive edit, but in the future that may not make sense.
 
 
-
-	auto FindExtrudeDistanceUsingInteractor = [](
-		FMeshEditorMode* This,
-		UViewportInteractor* ViewportInteractor,
-		const UEditableMesh* EditableMesh,
-		const FVector AxisOrigin,
-		const FVector AxisDirection,
-		const float AxisLength,
-		float& OutExtrudeDistance )
-	{
-		check( ViewportInteractor != nullptr );
-
-		OutExtrudeDistance = 0.0f;
-
-		float ClosestDistanceToRay = MAX_flt;
-		float BestExtrudeDistance = 0.0f;
-		
-
-		// Find the interactor ray that is closest to the axis line, and determine the distance from the origin of the
-		// axis to that point
-		FMeshEditorInteractorData& MeshEditorInteractorData = This->GetMeshEditorInteractorData( ViewportInteractor );
-		if( MeshEditorInteractorData.bLaserIsValid )	// @todo grabber: Grabber sphere impl neeed
-		{
-			const FVector AxisSegmentStart = AxisOrigin - AxisDirection * ( AxisLength * 0.5f );
-			const FVector AxisSegmentEnd = AxisOrigin + AxisDirection * ( AxisLength * 0.5f );
-
-			FVector ClosestPointOnAxis, ClosestPointOnRay;
-			FMath::SegmentDistToSegmentSafe(
-				AxisSegmentStart, AxisSegmentEnd,
-				MeshEditorInteractorData.LaserStart, MeshEditorInteractorData.LaserEnd,
-				/* Out */ ClosestPointOnAxis,
-				/* Out */ ClosestPointOnRay );
-			const float DistanceToRay = ( ClosestPointOnAxis - ClosestPointOnRay ).Size();
-
-			if( DistanceToRay < ClosestDistanceToRay )
-			{
-				ClosestDistanceToRay = DistanceToRay;
-
-				const FVector AxisToClosestPointVector = ( ClosestPointOnAxis - AxisOrigin );
-				BestExtrudeDistance = AxisToClosestPointVector.Size();
-
-				// Check to see if the closest point is actually behind the origin of the axis (negative extrusion distance)
-				const FVector ClosestPointDirection = AxisToClosestPointVector.GetSafeNormal();
-				if( FVector::DotProduct( ClosestPointDirection, AxisDirection ) < 0.0f )
-				{
-					BestExtrudeDistance *= -1.0f;
-				}
-			}
-		}
-
-		OutExtrudeDistance = BestExtrudeDistance;
-	};
-
-
 	// If this is an interim change, then everything that happens here -- all changes to our meshes, and even selection
 	// changes -- are guaranteed to be rolled back at the beginning of the next frame.  So we'll intercept any requests
 	// to store Undo history, and instead store those in a separate array to be processed ourselves next frame.
@@ -3415,113 +3354,6 @@ void FMeshEditorMode::UpdateActiveAction( const bool bIsActionFinishing )
 			}
 		}
 	}
-
-	else if( ActiveAction == EMeshEditAction::ExtrudePolygon ||
-			 ActiveAction == EMeshEditAction::FreelyExtrudePolygon )
-	{
-		// Extrude polygon
-		static TMap< UEditableMesh*, TArray< FMeshElement > > MeshesWithPolygonsToExtrude;
-		GetSelectedMeshesAndPolygons( /* Out */ MeshesWithPolygonsToExtrude );
-
-		if( MeshesWithPolygonsToExtrude.Num() > 0 )
-		{
-			if( ActiveAction == EMeshEditAction::FreelyExtrudePolygon )
-			{
-				// Start moving the extruded polygon right away!
-				bIsMovingSelectedMeshElements = true;
-			}
-
-			// Deselect the mesh elements before we delete them.  This will make sure they become selected again after undo.
-			{
-				FCompoundChangeInput CompoundRevertInput;
-
-				FSelectOrDeselectMeshElementsChangeInput ChangeInput;
-				for( auto& MeshAndPolygons : MeshesWithPolygonsToExtrude )
-				{
-					UEditableMesh* EditableMesh = MeshAndPolygons.Key;
-					const TArray<FMeshElement>& PolygonsToExtrude = MeshAndPolygons.Value;
-
-					for( const FMeshElement& PolygonElement : PolygonsToExtrude )
-					{
-						ChangeInput.MeshElementsToDeselect.Add( PolygonElement );
-					}
-				}
-
-				CompoundRevertInput.Subchanges.Add( MoveTemp( FSelectOrDeselectMeshElementsChange( MoveTemp( ChangeInput ) ).Execute( MeshEditorModeProxyObject ) ) );
-
-				TrackUndo( MeshEditorModeProxyObject, MakeUnique<FCompoundChange>( MoveTemp( CompoundRevertInput ) ) );
-			}
-
-			for( auto& MeshAndPolygons : MeshesWithPolygonsToExtrude )
-			{
-				UEditableMesh* EditableMesh = MeshAndPolygons.Key;
-				const TArray<FMeshElement>& PolygonsToExtrude = MeshAndPolygons.Value;
-
-				static TArray<FPolygonRef> PolygonRefsToExtrude;
-				PolygonRefsToExtrude.Reset();
-				for( const FMeshElement& PolygonToExtrude : PolygonsToExtrude )
-				{
-					FPolygonRef PolygonRef( PolygonToExtrude.ElementAddress.SectionID, FPolygonID( PolygonToExtrude.ElementAddress.ElementID ) );
-					PolygonRefsToExtrude.Add( PolygonRef );
-				}
-
-
-				float ExtrudeDistance = 0.0f;
-
-				if( ActiveAction == EMeshEditAction::ExtrudePolygon )
-				{
-					// Figure out how far the extruded polygon should be from where it started
-					const float AxisLength = 10000.0f;	// @todo mesheditor tweak (ideally should be infinite)
-					FindExtrudeDistanceUsingInteractor(
-						this,
-						ActiveActionInteractor,
-						EditableMesh,
-						ExtrudePolygonAxisOrigin,
-						ExtrudePolygonAxisDirection,
-						AxisLength,
-						/* Out */ ExtrudeDistance );
-				}
-
-
-				{
-					verify( !EditableMesh->AnyChangesToUndo() );
-
-					// Position the new polygon to where the interactor is
-					UPrimitiveComponent* Component = PolygonsToExtrude[ 0 ].Component.Get();	// NOTE: All polygons in this array belong to the same mesh/component, so we just need the first element
-					check( Component != nullptr );
-					// @todo mesheditor: We're working with a float here, so we'll treat the component scale as a scalar (X)
-					const float ComponentSpaceExtrudeDistance = ExtrudeDistance / Component->ComponentToWorld.GetScale3D().X;
-
-					// Create a copy of the polygon with new extruded polygons for each edge
-					static TArray<FPolygonRef> NewExtrudedFrontPolygons;
-					NewExtrudedFrontPolygons.Reset();
-					EditableMesh->ExtrudePolygons( PolygonRefsToExtrude, ComponentSpaceExtrudeDistance, /* Out */ NewExtrudedFrontPolygons );
-
-
-					// Make sure the new polygons are selected.  The old polygon was deleted and will become deselected automatically.
-					for( int32 NewPolygonNumber = 0; NewPolygonNumber < NewExtrudedFrontPolygons.Num(); ++NewPolygonNumber )
-					{
-						const FPolygonRef& NewExtrudedFrontPolygon = NewExtrudedFrontPolygons[ NewPolygonNumber ];
-
-						const FMeshElement& MeshElement = PolygonsToExtrude[ NewPolygonNumber ];
-
-						FMeshElement NewExtrudedPolygonMeshElement;
-						{
-							NewExtrudedPolygonMeshElement.Component = MeshElement.Component;
-							NewExtrudedPolygonMeshElement.ElementAddress = MeshElement.ElementAddress;
-							NewExtrudedPolygonMeshElement.ElementAddress.SectionID = NewExtrudedFrontPolygon.SectionID;
-							NewExtrudedPolygonMeshElement.ElementAddress.ElementID = NewExtrudedFrontPolygon.PolygonID;
-						}
-
-						// Queue selection of this new element.  We don't want it to be part of the current action.
-						MeshElementsToSelect.Add( NewExtrudedPolygonMeshElement );
-					}
-				}
-
-				TrackUndo( EditableMesh, EditableMesh->MakeUndo() );
-			}
-		}
-	}
 	else if( ActiveAction == EMeshEditAction::ExtendEdge )
 	{
 		// Extend edge
@@ -4910,60 +4742,6 @@ void FMeshEditorMode::OnViewportInteractionInputAction( FEditorViewportClient& V
 						const bool bActionNeedsHoverLocation = false;
 						StartAction( EMeshEditAction::Move, ViewportInteractor, bActionNeedsHoverLocation, LOCTEXT( "UndoDragPolygon", "Drag Polygon" ) );
 					}
-					// @todo mesheditor: Need a "Shift+Click+Drag" extrude instead of having to explicitly equip the action
-					else if( SelectedMeshElementType == EEditableMeshElementType::Polygon && EquippedPolygonAction == EMeshEditAction::ExtrudePolygon )
-					{
-						bool bHaveExtrudeAxis = false;
-						{
-							const FMeshElement& PolygonElement = MeshEditorInteractorData.HoveredMeshElement;
-
-							if( PolygonElement.IsValidMeshElement() &&
-								PolygonElement.ElementAddress.ElementType == EEditableMeshElementType::Polygon )
-							{
-								// Is it selected?
-								if( IsMeshElementSelected( PolygonElement ) )
-								{
-									UEditableMesh* EditableMesh = FindOrCreateEditableMesh( *PolygonElement.Component, PolygonElement.ElementAddress.SubMeshAddress );
-									if( EditableMesh != nullptr )
-									{
-										if( IsElementIDValid( PolygonElement, EditableMesh ) )
-										{
-											this->ExtrudePolygonAxisOrigin = MeshEditorInteractorData.HoverLocation;
-
-											// Use the polygon normal as the extrude axis direction
-											{
-												UPrimitiveComponent* Component = PolygonElement.Component.Get();
-												check( Component != nullptr );
-												const FMatrix ComponentToWorldMatrix = Component->GetRenderMatrix();
-
-												const FPolygonRef PolygonRef( PolygonElement.ElementAddress.SectionID, FPolygonID( PolygonElement.ElementAddress.ElementID ) );
-												const FVector ComponentSpacePolygonNormal = EditableMesh->ComputePolygonNormal( PolygonRef );
-
-												const FVector WorldSpacePolygonNormal = ComponentToWorldMatrix.TransformVector( ComponentSpacePolygonNormal ).GetSafeNormal();
-
-												this->ExtrudePolygonAxisDirection = WorldSpacePolygonNormal;
-											}
-
-											bHaveExtrudeAxis = true;
-										}
-									}
-								}
-							}
-						}
-
-						if( bHaveExtrudeAxis )
-						{
-							const bool bActionNeedsHoverLocation = false;
-							StartAction( EMeshEditAction::ExtrudePolygon, ViewportInteractor, bActionNeedsHoverLocation, LOCTEXT( "UndoExtrudePolygon", "Extrude Polygon" ) );
-							bOutIsInputCaptured = true;
-						}
-					}
-					else if( SelectedMeshElementType == EEditableMeshElementType::Polygon && EquippedPolygonAction == EMeshEditAction::FreelyExtrudePolygon )
-					{
-						bWantToStartMoving = true;
-						const bool bActionNeedsHoverLocation = false;
-						StartAction( EMeshEditAction::FreelyExtrudePolygon, ViewportInteractor, bActionNeedsHoverLocation, LOCTEXT( "UndoFreelyExtrudePolygon", "Freely Extrude Polygon" ) );
-					}
 					else if( SelectedMeshElementType == EEditableMeshElementType::Edge && EquippedEdgeAction == EMeshEditAction::SplitEdge )
 					{
 						const bool bActionNeedsHoverLocation = true;
@@ -5788,19 +5566,6 @@ void FMeshEditorMode::MakeVRRadialMenuActionsMenu(FMenuBuilder& MenuBuilder, TSh
 			NAME_None,
 			EUserInterfaceActionType::CollapsedButton
 			);
-		MenuBuilder.AddMenuEntry(
-			LOCTEXT("Extrude", "Extrude"),
-			FText(),
-			FSlateIcon(FMeshEditorStyle::GetStyleSetName(), "MeshEditorMode.PolyExtrude"),
-			FUIAction
-			(
-				FExecuteAction::CreateLambda([this] { SetEquippedAction( EEditableMeshElementType::Polygon, EMeshEditAction::ExtrudePolygon ); }),
-				FCanExecuteAction::CreateSP(this, &FMeshEditorMode::IsMeshElementTypeSelectedOrIsActiveSelectionMode, EEditableMeshElementType::Polygon),
-				FIsActionChecked::CreateLambda([this] { return (EquippedPolygonAction == EMeshEditAction::ExtrudePolygon); })
-			),
-			NAME_None,
-			EUserInterfaceActionType::ToggleButton
-			);
 	}
 	else if (GetMeshElementSelectionMode() == EEditableMeshElementType::Edge)
 	{
diff --git a/Engine/Plugins/Editor/MeshEditor/Source/MeshEditor/MeshEditorMode.h b/Engine/Plugins/Editor/MeshEditor/Source/MeshEditor/MeshEditorMode.h
--- a/Engine/Plugins/Editor/MeshEditor/Source/MeshEditor/MeshEditorMode.h
+++ b/Engine/Plugins/Editor/MeshEditor/Source/MeshEditor/MeshEditorMode.h
@@ -585,16 +585,6 @@ protected:
 	TArray<TTuple<TSharedPtr<FUICommandInfo>, FUIAction>> EdgeActions;
 	TArray<TTuple<TSharedPtr<FUICommandInfo>, FUIAction>> PolygonActions;
 
-	//
-	// ExtrudePolygon
-	//
-
-	/** When extruding polygons, we need to keep track of the interactor's impact point and normal, because
-	    the user is going to be aiming their interactor along that axis to choose an extrusion point */
-	FVector ExtrudePolygonAxisOrigin;
-	FVector ExtrudePolygonAxisDirection;
-
-
 	//
 	// SplitEdgeAndDragVertex
 	//
diff --git a/Engine/Plugins/Editor/MeshEditor/Source/MeshEditor/Public/IMeshEditorModeEditingContract.h b/Engine/Plugins/Editor/MeshEditor/Source/MeshEditor/Public/IMeshEditorModeEditingContract.h
--- a/Engine/Plugins/Editor/MeshEditor/Source/MeshEditor/Public/IMeshEditorModeEditingContract.h
+++ b/Engine/Plugins/Editor/MeshEditor/Source/MeshEditor/Public/IMeshEditorModeEditingContract.h
@@ -29,12 +29,6 @@ namespace EMeshEditAction
 	/** Splits an edge by inserting a new vertex, then immediately starts dragging that vertex */
 	MESHEDITOR_API extern const FName SplitEdgeAndDragVertex;
 
-	/** Extrude polygon by making a copy of it and allowing you to shift it along the polygon normal axis */
-	MESHEDITOR_API extern const FName ExtrudePolygon;
-
-	/** Extrude polygon by making a copy of it and allowing you to move it around freely */
-	MESHEDITOR_API extern const FName FreelyExtrudePolygon;
-
 	/** Extend an edge by making a copy of it and allowing you to move it around */
 	MESHEDITOR_API extern const FName ExtendEdge;
 
diff --git a/Engine/Plugins/Editor/MeshEditor/Source/MeshEditor/Public/MeshEditorCommands.h b/Engine/Plugins/Editor/MeshEditor/Source/MeshEditor/Public/MeshEditorCommands.h
--- a/Engine/Plugins/Editor/MeshEditor/Source/MeshEditor/Public/MeshEditorCommands.h
+++ b/Engine/Plugins/Editor/MeshEditor/Source/MeshEditor/Public/MeshEditorCommands.h
@@ -287,12 +287,6 @@ public:
 	/** Sets the primary action to move polygons */
 	TSharedPtr<FUICommandInfo> MovePolygon;
 
-	/** Sets the primary action to extrude polygons */
-	TSharedPtr<FUICommandInfo> ExtrudePolygon;
-
-	/** Sets the primary action to freely extrude polygons */
-	TSharedPtr<FUICommandInfo> FreelyExtrudePolygon;
-
 	/** Flips the currently selected polygon(s) */
 	TSharedPtr<FUICommandInfo> FlipPolygon;
 
diff --git a/Engine/Plugins/Editor/MeshEditor/Source/PolygonModeling/ExtrudePolygon.cpp b/Engine/Plugins/Editor/MeshEditor/Source/PolygonModeling/ExtrudePolygon.cpp
new file mode 100644
--- /dev/null
+++ b/Engine/Plugins/Editor/MeshEditor/Source/PolygonModeling/ExtrudePolygon.cpp
@@ -0,0 +1,216 @@
+// Copyright 1998-2016 Epic Games, Inc. All Rights Reserved.
+
+#include "ExtrudePolygon.h"
+#include "IMeshEditorModeEditingContract.h"
+#include "IMeshEditorModeUIContract.h"
+#include "UICommandInfo.h"
+#include "EditableMesh.h"
+#include "MeshElement.h"
+#include "MultiBoxBuilder.h"
+#include "UICommandList.h"
+#include "ViewportInteractor.h"
+
+#define LOCTEXT_NAMESPACE "MeshEditorMode"
+
+
+void UExtrudePolygonCommand::FindExtrudeDistanceUsingInteractor(
+	IMeshEditorModeEditingContract& MeshEditorMode,
+	UViewportInteractor* ViewportInteractor,
+	const UEditableMesh* EditableMesh,
+	const FVector AxisOrigin,
+	const FVector AxisDirection,
+	const float AxisLength,
+	float& OutExtrudeDistance )
+{
+	check( ViewportInteractor != nullptr );
+
+	OutExtrudeDistance = 0.0f;
+
+	float ClosestDistanceToRay = MAX_flt;
+	float BestExtrudeDistance = 0.0f;
+
+
+	// Find the interactor ray that is closest to the axis line, and determine the distance from the origin of the
+	// axis to that point
+	// @todo mesheditor grabber: Needs grabber sphere support
+	FVector LaserStart, LaserEnd;
+	const bool bLaserIsValid = ViewportInteractor->GetLaserPointer( /* Out */ LaserStart, /* Out */ LaserEnd );
+	if( bLaserIsValid )
+	{
+		const FVector AxisSegmentStart = AxisOrigin - AxisDirection * ( AxisLength * 0.5f );
+		const FVector AxisSegmentEnd = AxisOrigin + AxisDirection * ( AxisLength * 0.5f );
+
+		FVector ClosestPointOnAxis, ClosestPointOnRay;
+		FMath::SegmentDistToSegmentSafe(
+			AxisSegmentStart, AxisSegmentEnd,
+			LaserStart, LaserEnd,
+			/* Out */ ClosestPointOnAxis,
+			/* Out */ ClosestPointOnRay );
+		const float DistanceToRay = ( ClosestPointOnAxis - ClosestPointOnRay ).Size();
+
+		if( DistanceToRay < ClosestDistanceToRay )
+		{
+			ClosestDistanceToRay = DistanceToRay;
+
+			const FVector AxisToClosestPointVector = ( ClosestPointOnAxis - AxisOrigin );
+			BestExtrudeDistance = AxisToClosestPointVector.Size();
+
+			// Check to see if the closest point is actually behind the origin of the axis (negative extrusion distance)
+			const FVector ClosestPointDirection = AxisToClosestPointVector.GetSafeNormal();
+			if( FVector::DotProduct( ClosestPointDirection, AxisDirection ) < 0.0f )
+			{
+				BestExtrudeDistance *= -1.0f;
+			}
+		}
+	}
+
+	OutExtrudeDistance = BestExtrudeDistance;
+};
+
+
+
+void UExtrudePolygonCommand::RegisterUICommand( FBindingContext* BindingContext )
+{
+	UI_COMMAND_EXT( BindingContext, /* Out */ UICommandInfo, "ExtrudePolygon", "Extrude Polygon Mode", "Set the primary action to extrude polygons.", EUserInterfaceActionType::RadioButton, FInputChord() );
+}
+
+
+bool UExtrudePolygonCommand::TryStartingToDrag( IMeshEditorModeEditingContract& MeshEditorMode, UViewportInteractor* ViewportInteractor )
+{
+	// @todo mesheditor: Need a "Shift+Click+Drag" extrude instead of having to explicitly equip the action
+	bool bHaveExtrudeAxis = false;
+	{
+		const FMeshElement& PolygonElement = MeshEditorMode.GetHoveredMeshElement( ViewportInteractor );
+		if( PolygonElement.ElementAddress.ElementType == EEditableMeshElementType::Polygon )
+		{
+			// Is it selected?
+			if( MeshEditorMode.IsMeshElementSelected( PolygonElement ) )
+			{
+				const UEditableMesh* EditableMesh = MeshEditorMode.FindEditableMesh( *PolygonElement.Component, PolygonElement.ElementAddress.SubMeshAddress );
+				if( EditableMesh != nullptr )
+				{
+					this->ExtrudePolygonAxisOrigin = ViewportInteractor->GetHoverLocation();
+
+					// Use the polygon normal as the extrude axis direction
+					{
+						UPrimitiveComponent* Component = PolygonElement.Component.Get();
+						check( Component != nullptr );
+						const FMatrix ComponentToWorldMatrix = Component->GetRenderMatrix();
+
+						const FPolygonRef PolygonRef( PolygonElement.ElementAddress.SectionID, FPolygonID( PolygonElement.ElementAddress.ElementID ) );
+						const FVector ComponentSpacePolygonNormal = EditableMesh->ComputePolygonNormal( PolygonRef );
+
+						const FVector WorldSpacePolygonNormal = ComponentToWorldMatrix.TransformVector( ComponentSpacePolygonNormal ).GetSafeNormal();
+
+						this->ExtrudePolygonAxisDirection = WorldSpacePolygonNormal;
+					}
+
+					bHaveExtrudeAxis = true;
+				}
+			}
+		}
+	}
+
+	return bHaveExtrudeAxis;
+}
+
+
+void UExtrudePolygonCommand::ApplyDuringDrag( IMeshEditorModeEditingContract& MeshEditorMode, UViewportInteractor* ViewportInteractor, bool& bOutShouldDeselectAllFirst, TArray<FMeshElement>& OutMeshElementsToSelect )
+{
+	static TMap< UEditableMesh*, TArray< FMeshElement > > MeshesWithPolygonsToExtrude;
+	MeshEditorMode.GetSelectedMeshesAndPolygons( /* Out */ MeshesWithPolygonsToExtrude );
+
+	if( MeshesWithPolygonsToExtrude.Num() > 0 )
+	{
+		// Deselect the mesh elements before we delete them.  This will make sure they become selected again after undo.
+		MeshEditorMode.DeselectMeshElements( MeshesWithPolygonsToExtrude );
+
+		for( auto& MeshAndPolygons : MeshesWithPolygonsToExtrude )
+		{
+			UEditableMesh* EditableMesh = MeshAndPolygons.Key;
+			const TArray<FMeshElement>& PolygonsToExtrude = MeshAndPolygons.Value;
+
+			static TArray<FPolygonRef> PolygonRefsToExtrude;
+			PolygonRefsToExtrude.Reset();
+			for( const FMeshElement& PolygonToExtrude : PolygonsToExtrude )
+			{
+				FPolygonRef PolygonRef( PolygonToExtrude.ElementAddress.SectionID, FPolygonID( PolygonToExtrude.ElementAddress.ElementID ) );
+				PolygonRefsToExtrude.Add( PolygonRef );
+			}
+
+
+			// Figure out how far the extruded polygon should be from where it started
+			float ExtrudeDistance = 0.0f;
+			{
+				const float AxisLength = 10000.0f;	// @todo mesheditor tweak (ideally should be infinite)
+				FindExtrudeDistanceUsingInteractor(
+					MeshEditorMode,
+					ViewportInteractor,
+					EditableMesh,
+					ExtrudePolygonAxisOrigin,
+					ExtrudePolygonAxisDirection,
+					AxisLength,
+					/* Out */ ExtrudeDistance );
+			}
+
+
+			{
+				verify( !EditableMesh->AnyChangesToUndo() );
+
+				// Position the new polygon to where the interactor is
+				UPrimitiveComponent* Component = PolygonsToExtrude[ 0 ].Component.Get();	// NOTE: All polygons in this array belong to the same mesh/component, so we just need the first element
+				check( Component != nullptr );
+				// @todo mesheditor: We're working with a float here, so we'll treat the component scale as a scalar (X)
+				const float ComponentSpaceExtrudeDistance = ExtrudeDistance / Component->ComponentToWorld.GetScale3D().X;
+
+				// Create a copy of the polygon with new extruded polygons for each edge
+				static TArray<FPolygonRef> NewExtrudedFrontPolygons;
+				NewExtrudedFrontPolygons.Reset();
+				EditableMesh->ExtrudePolygons( PolygonRefsToExtrude, ComponentSpaceExtrudeDistance, /* Out */ NewExtrudedFrontPolygons );
+
+
+				// Make sure the new polygons are selected.  The old polygon was deleted and will become deselected automatically.
+				for( int32 NewPolygonNumber = 0; NewPolygonNumber < NewExtrudedFrontPolygons.Num(); ++NewPolygonNumber )
+				{
+					const FPolygonRef& NewExtrudedFrontPolygon = NewExtrudedFrontPolygons[ NewPolygonNumber ];
+
+					const FMeshElement& MeshElement = PolygonsToExtrude[ NewPolygonNumber ];
+
+					FMeshElement NewExtrudedPolygonMeshElement;
+					{
+						NewExtrudedPolygonMeshElement.Component = MeshElement.Component;
+						NewExtrudedPolygonMeshElement.ElementAddress = MeshElement.ElementAddress;
+						NewExtrudedPolygonMeshElement.ElementAddress.SectionID = NewExtrudedFrontPolygon.SectionID;
+						NewExtrudedPolygonMeshElement.ElementAddress.ElementID = NewExtrudedFrontPolygon.PolygonID;
+					}
+
+					// Queue selection of this new element.  We don't want it to be part of the current action.
+					OutMeshElementsToSelect.Add( NewExtrudedPolygonMeshElement );
+				}
+			}
+
+			MeshEditorMode.TrackUndo( EditableMesh, EditableMesh->MakeUndo() );
+		}
+	}
+}
+
+
+void UExtrudePolygonCommand::AddToVRRadialMenuActionsMenu( IMeshEditorModeUIContract& MeshEditorMode, FMenuBuilder& MenuBuilder, TSharedPtr<FUICommandList> CommandList, const FName TEMPHACK_StyleSetName, class UVREditorMode* VRMode )
+{
+	if( MeshEditorMode.GetMeshElementSelectionMode() == EEditableMeshElementType::Polygon )
+	{
+		MenuBuilder.AddMenuEntry(
+			LOCTEXT( "VRExtrudePolygon", "Extrude" ),
+			FText(),
+			FSlateIcon( TEMPHACK_StyleSetName, "MeshEditorMode.PolyExtrude" ),
+			MakeUIAction( MeshEditorMode ),
+			NAME_None,
+			EUserInterfaceActionType::ToggleButton
+		);
+	}
+}
+
+
+
+#undef LOCTEXT_NAMESPACE
+
diff --git a/Engine/Plugins/Editor/MeshEditor/Source/PolygonModeling/ExtrudePolygon.h b/Engine/Plugins/Editor/MeshEditor/Source/PolygonModeling/ExtrudePolygon.h
new file mode 100644
--- /dev/null
+++ b/Engine/Plugins/Editor/MeshEditor/Source/PolygonModeling/ExtrudePolygon.h
@@ -0,0 +1,50 @@
+// Copyright 1998-2016 Epic Games, Inc. All Rights Reserved.
+
+#pragma once
+
+#include "MeshEditorCommands.h"
+#include "ExtrudePolygon.generated.h"
+
+
+/** Extrudes the polygon along an axis */
+UCLASS()
+class UExtrudePolygonCommand : public UMeshEditorPolygonCommand
+{
+	GENERATED_BODY()
+
+protected:
+
+	UExtrudePolygonCommand()
+		: ExtrudePolygonAxisOrigin( FVector::ZeroVector ),
+		  ExtrudePolygonAxisDirection( FVector::ZeroVector )
+	{
+		bIsMode = true;
+		UndoText = NSLOCTEXT( "MeshEditor", "UndoExtrudePolygon", "Extrude Polygon" );
+		bNeedsHoverLocation = false;
+		bNeedsDraggingInitiated = false;
+	}
+
+	/** Figures out how far to extrude the polygon based on where the interactor is aiming */
+	void FindExtrudeDistanceUsingInteractor(
+		class IMeshEditorModeEditingContract& MeshEditorMode,
+		class UViewportInteractor* ViewportInteractor,
+		const class UEditableMesh* EditableMesh,
+		const FVector AxisOrigin,
+		const FVector AxisDirection,
+		const float AxisLength,
+		float& OutExtrudeDistance );
+
+	// Overrides
+	virtual void RegisterUICommand( class FBindingContext* BindingContext ) override;
+	virtual bool TryStartingToDrag( IMeshEditorModeEditingContract& MeshEditorMode, class UViewportInteractor* ViewportInteractor ) override;
+	virtual void ApplyDuringDrag( IMeshEditorModeEditingContract& MeshEditorMode, class UViewportInteractor* ViewportInteractor, bool& bOutShouldDeselectAllFirst, TArray<FMeshElement>& OutMeshElementsToSelect ) override;
+	virtual void AddToVRRadialMenuActionsMenu( class IMeshEditorModeUIContract& MeshEditorMode, class FMenuBuilder& MenuBuilder, TSharedPtr<FUICommandList> CommandList, const FName TEMPHACK_StyleSetName, class UVREditorMode* VRMode ) override;
+
+protected:
+
+	/** When extruding polygons, we need to keep track of the interactor's impact point and normal, because
+	    the user is going to be aiming their interactor along that axis to choose an extrusion point */
+	FVector ExtrudePolygonAxisOrigin;
+	FVector ExtrudePolygonAxisDirection;
+
+};
