Mesh Editor extensibility (phase 5)
- Split Edge, Extend Edge and Extend Vertex have been migrated to a new module
- Added support for extensible mesh commands that drag mesh elements immediately after an action
- Fixed selection undo events added with empty selection sets
- Simplified extensibility API for mesh commands that apply during a drag

diff --git a/Engine/Plugins/Editor/MeshEditor/Source/MeshEditor/MeshEditorCommands.cpp b/Engine/Plugins/Editor/MeshEditor/Source/MeshEditor/MeshEditorCommands.cpp
--- a/Engine/Plugins/Editor/MeshEditor/Source/MeshEditor/MeshEditorCommands.cpp
+++ b/Engine/Plugins/Editor/MeshEditor/Source/MeshEditor/MeshEditorCommands.cpp
@@ -81,7 +81,6 @@ FMeshEditorVertexCommands::FMeshEditorVertexCommands()
 void FMeshEditorVertexCommands::RegisterCommands()
 {
 	UI_COMMAND(MoveVertex, "Move Vertex Mode", "Set the primary action to move vertices.", EUserInterfaceActionType::RadioButton, FInputChord(EKeys::F1));
-	UI_COMMAND(ExtendVertex, "Extend Vertex Mode", "Set the primary action to extend vertices.", EUserInterfaceActionType::RadioButton, FInputChord(EKeys::F2));
 	UI_COMMAND(EditVertexCornerSharpness, "Edit Vertex Corner Sharpness Mode", "Set the primary action to edit the vertex's subdivision corner sharpness amount.", EUserInterfaceActionType::RadioButton, FInputChord(EKeys::F6));
 
 	UI_COMMAND(RemoveVertex, "Remove Vertex", "Remove the selected vertex if possible.", EUserInterfaceActionType::Button, FInputChord(EKeys::BackSpace));
@@ -111,9 +110,6 @@ FMeshEditorEdgeCommands::FMeshEditorEdgeCommands()
 void FMeshEditorEdgeCommands::RegisterCommands()
 {
 	UI_COMMAND(MoveEdge, "Move Edge Mode", "Set the primary action to move edges.", EUserInterfaceActionType::RadioButton, FInputChord(EKeys::F1));
-	UI_COMMAND(SplitEdge, "Split Edge Mode", "Set the primary action to split edges.", EUserInterfaceActionType::RadioButton, FInputChord(EKeys::F2));
-	UI_COMMAND(SplitEdgeAndDragVertex, "Split Edge and Drag Vertex Mode", "Set the primary action to split edges and drag vertices.", EUserInterfaceActionType::RadioButton, FInputChord(EKeys::F3));
-	UI_COMMAND(ExtendEdge, "Extend Edge Mode", "Set the primary action to extend edges.", EUserInterfaceActionType::RadioButton, FInputChord(EKeys::F5));
 	UI_COMMAND(EditEdgeCreaseSharpness, "Edit Edge Crease Sharpness Mode", "Set the primary action to edit the edge's crease sharpness.", EUserInterfaceActionType::RadioButton, FInputChord(EKeys::F6));
 
 	UI_COMMAND(RemoveEdge, "Remove Edge", "Remove the selected edge if possible.", EUserInterfaceActionType::Button, FInputChord(EKeys::BackSpace));
diff --git a/Engine/Plugins/Editor/MeshEditor/Source/MeshEditor/MeshEditorMode.cpp b/Engine/Plugins/Editor/MeshEditor/Source/MeshEditor/MeshEditorMode.cpp
--- a/Engine/Plugins/Editor/MeshEditor/Source/MeshEditor/MeshEditorMode.cpp
+++ b/Engine/Plugins/Editor/MeshEditor/Source/MeshEditor/MeshEditorMode.cpp
@@ -47,10 +47,6 @@ namespace EMeshEditAction
 	const Type SelectByPainting( "SelectByPainting" );
 	const Type MoveUsingGizmo( "MoveUsingGizmo" );
 	const Type Move( "Move" );
-	const Type SplitEdge( "SplitEdge" );
-	const Type SplitEdgeAndDragVertex( "SplitEdgeAndDragVertex" );
-	const Type ExtendEdge( "ExtendEdge" );
-	const Type ExtendVertex( "ExtendVertex" );
 	const Type EditVertexCornerSharpness( "EditVertexCornerSharpness" );
 	const Type EditEdgeCreaseSharpness( "EditEdgeCreaseSharpness" );
 	const Type DrawVertices( "DrawVertices" );
@@ -392,8 +388,6 @@ FMeshEditorMode::FMeshEditorMode()
 	  ActiveActionInteractor( nullptr ),
 	  bActiveActionNeedsHoverLocation( false ),
 	  bIsFirstActiveActionUpdate( false ),
-	  SplitEdgeMeshesAndEdgesToSplit(),
-	  SplitEdgeSplitList(),
 	  EditSharpnessStartLocation( FVector::ZeroVector ),
 	  SelectingByPaintingRevertChangeInput( nullptr ),
 	  bShowVertexNormals( false ),
@@ -527,13 +521,9 @@ void FMeshEditorMode::BindCommands()
 
 	// Register editing modes (equipped actions)
 	RegisterVertexEditingMode( MeshEditorVertexActions.MoveVertex, EMeshEditAction::Move );
-	RegisterVertexEditingMode( MeshEditorVertexActions.ExtendVertex, EMeshEditAction::ExtendVertex );
 	RegisterVertexEditingMode( MeshEditorVertexActions.EditVertexCornerSharpness, EMeshEditAction::EditVertexCornerSharpness );
 
 	RegisterEdgeEditingMode( MeshEditorEdgeActions.MoveEdge, EMeshEditAction::Move );
-	RegisterEdgeEditingMode( MeshEditorEdgeActions.SplitEdge, EMeshEditAction::SplitEdge );
-	RegisterEdgeEditingMode( MeshEditorEdgeActions.SplitEdgeAndDragVertex, EMeshEditAction::SplitEdgeAndDragVertex );
-	RegisterEdgeEditingMode( MeshEditorEdgeActions.ExtendEdge, EMeshEditAction::ExtendEdge );
 	RegisterEdgeEditingMode( MeshEditorEdgeActions.EditEdgeCreaseSharpness, EMeshEditAction::EditEdgeCreaseSharpness );
 
 	RegisterPolygonEditingMode( MeshEditorPolygonActions.MovePolygon, EMeshEditAction::Move );
@@ -1110,14 +1100,13 @@ void FMeshEditorMode::CommitEditableMeshIfNecessary( UEditableMesh* EditableMesh
 		// @todo mesheditor: this is a little bit fragile. Ideally we initialize these things after the new instance has been created.
 		// @todo mesheditor extensibility: Figure out how external FMeshElements can be fixed up (either a callback with FixUpMeshElement function access, or they are registered with this system?)
 		// FixUpMeshElement( InsetUsingPolygonElement );
-
-		for( auto& SplitEdgeMeshAndEdgesToSplit : SplitEdgeMeshesAndEdgesToSplit )
-		{
-			for( auto& EdgeToSplit : SplitEdgeMeshAndEdgesToSplit.Value )
-			{
-				FixUpMeshElement( EdgeToSplit );
-			}
-		}
+		// 		for( auto& SplitEdgeMeshAndEdgesToSplit : SplitEdgeMeshesAndEdgesToSplit )
+		// 		{
+		// 			for( auto& EdgeToSplit : SplitEdgeMeshAndEdgesToSplit.Value )
+		// 			{
+		// 				FixUpMeshElement( EdgeToSplit );
+		// 			}
+		// 		}
 
 		RefreshTransformables();
 	}
@@ -1193,37 +1182,37 @@ const UMeshEditorAssetContainer& FMeshEditorMode::GetAssetContainer() const
 
 void FMeshEditorMode::SelectMeshElements( const TArray<FMeshElement>& MeshElementsToSelect )
 {
-	FCompoundChangeInput CompoundRevertInput;
-	FSelectOrDeselectMeshElementsChangeInput ChangeInput;
-	ChangeInput.MeshElementsToSelect = MeshElementsToSelect;
-
-	CompoundRevertInput.Subchanges.Add( MoveTemp( FSelectOrDeselectMeshElementsChange( MoveTemp( ChangeInput ) ).Execute( MeshEditorModeProxyObject ) ) );
-
-	TrackUndo( MeshEditorModeProxyObject, MakeUnique<FCompoundChange>( MoveTemp( CompoundRevertInput ) ) );
+	if( MeshElementsToSelect.Num() > 0 )
+	{
+		FSelectOrDeselectMeshElementsChangeInput ChangeInput;
+		ChangeInput.MeshElementsToSelect = MeshElementsToSelect;
+		TrackUndo( MeshEditorModeProxyObject, FSelectOrDeselectMeshElementsChange( ChangeInput ).Execute( MeshEditorModeProxyObject ) );
+	}
 }
 
 
 void FMeshEditorMode::DeselectAllMeshElements()
 {
-	TrackUndo( MeshEditorModeProxyObject, FDeselectAllMeshElementsChange( FDeselectAllMeshElementsChangeInput() ).Execute( MeshEditorModeProxyObject ) );
+	if( SelectedMeshElements.Num() > 0 )
+	{
+		TrackUndo( MeshEditorModeProxyObject, FDeselectAllMeshElementsChange( FDeselectAllMeshElementsChangeInput() ).Execute( MeshEditorModeProxyObject ) );
+	}
 }
 
 
 void FMeshEditorMode::DeselectMeshElements( const TArray<FMeshElement>& MeshElementsToDeselect )
 {
-	FCompoundChangeInput CompoundRevertInput;
-	FSelectOrDeselectMeshElementsChangeInput ChangeInput;
-	ChangeInput.MeshElementsToDeselect = MeshElementsToDeselect;
-
-	CompoundRevertInput.Subchanges.Add( MoveTemp( FSelectOrDeselectMeshElementsChange( MoveTemp( ChangeInput ) ).Execute( MeshEditorModeProxyObject ) ) );
-
-	TrackUndo( MeshEditorModeProxyObject, MakeUnique<FCompoundChange>( MoveTemp( CompoundRevertInput ) ) );
+	if( MeshElementsToDeselect.Num() > 0 )
+	{
+		FSelectOrDeselectMeshElementsChangeInput ChangeInput;
+		ChangeInput.MeshElementsToDeselect = MeshElementsToDeselect;
+		TrackUndo( MeshEditorModeProxyObject, FSelectOrDeselectMeshElementsChange( ChangeInput ).Execute( MeshEditorModeProxyObject ) );
+	}
 }
 
 
 void FMeshEditorMode::DeselectMeshElements( const TMap<UEditableMesh*, TArray<FMeshElement>>& MeshElementsToDeselect )
 {
-	FCompoundChangeInput CompoundRevertInput;
 	FSelectOrDeselectMeshElementsChangeInput ChangeInput;
 	for( const auto& MeshAndElements : MeshElementsToDeselect )
 	{
@@ -1232,10 +1221,10 @@ void FMeshEditorMode::DeselectMeshElements( const TMap<UEditableMesh*, TArray<FM
 			ChangeInput.MeshElementsToDeselect.Add( MeshElementToDeselect );
 		}
 	}
-
-	CompoundRevertInput.Subchanges.Add( MoveTemp( FSelectOrDeselectMeshElementsChange( MoveTemp( ChangeInput ) ).Execute( MeshEditorModeProxyObject ) ) );
-
-	TrackUndo( MeshEditorModeProxyObject, MakeUnique<FCompoundChange>( MoveTemp( CompoundRevertInput ) ) );
+	if( ChangeInput.MeshElementsToDeselect.Num() > 0 )
+	{
+		TrackUndo( MeshEditorModeProxyObject, FSelectOrDeselectMeshElementsChange( ChangeInput ).Execute( MeshEditorModeProxyObject ) );
+	}
 }
 
 
@@ -1737,15 +1726,7 @@ bool FMeshEditorMode::RemoveSelectedVertices()
 	{
 		// Make sure we're not still hovering over a polygon we removed
 		ClearInvalidSelectedElements();
-
-		FCompoundChangeInput CompoundRevertInput;
-
-		// Select the edge that the vertex we removed was a member of
-		FSelectOrDeselectMeshElementsChangeInput ChangeInput;
-		ChangeInput.MeshElementsToSelect = MeshElementsToSelect;
-		CompoundRevertInput.Subchanges.Add( MoveTemp( FSelectOrDeselectMeshElementsChange( MoveTemp( ChangeInput ) ).Execute( MeshEditorModeProxyObject ) ) );
-
-		TrackUndo( MeshEditorModeProxyObject, MakeUnique<FCompoundChange>( MoveTemp( CompoundRevertInput ) ) );
+		SelectMeshElements( MeshElementsToSelect );
 	}
 
 	return true;
@@ -1831,15 +1812,7 @@ bool FMeshEditorMode::WeldSelectedVertices()
 	{
 		// Make sure we're not still hovering over a polygon we removed
 		ClearInvalidSelectedElements();
-
-		FCompoundChangeInput CompoundRevertInput;
-
-		// Select the vertex that was welded.
-		FSelectOrDeselectMeshElementsChangeInput ChangeInput;
-		ChangeInput.MeshElementsToSelect = MeshElementsToSelect;
-		CompoundRevertInput.Subchanges.Add( MoveTemp( FSelectOrDeselectMeshElementsChange( MoveTemp( ChangeInput ) ).Execute( MeshEditorModeProxyObject ) ) );
-
-		TrackUndo( MeshEditorModeProxyObject, MakeUnique<FCompoundChange>( MoveTemp( CompoundRevertInput ) ) );
+		SelectMeshElements( MeshElementsToSelect );
 	}
 
 	return true;
@@ -1977,14 +1950,7 @@ bool FMeshEditorMode::TriangulateSelectedPolygons()
 	{
 		// Make sure we're not still hovering over a polygon we removed
 		ClearInvalidSelectedElements();
-
-		FCompoundChangeInput CompoundRevertInput;
-
-		FSelectOrDeselectMeshElementsChangeInput ChangeInput;
-		ChangeInput.MeshElementsToSelect = MeshElementsToSelect;
-		CompoundRevertInput.Subchanges.Add( MoveTemp( FSelectOrDeselectMeshElementsChange( MoveTemp( ChangeInput ) ).Execute( MeshEditorModeProxyObject ) ) );
-
-		TrackUndo( MeshEditorModeProxyObject, MakeUnique<FCompoundChange>( MoveTemp( CompoundRevertInput ) ) );
+		SelectMeshElements( MeshElementsToSelect );
 	}
 
 
@@ -2089,14 +2055,7 @@ bool FMeshEditorMode::AssignMaterialToSelectedPolygons( UMaterialInterface* Sele
 		{
 			// Make sure we're not still hovering over a polygon we removed
 			ClearInvalidSelectedElements();
-
-			FCompoundChangeInput CompoundRevertInput;
-
-			FSelectOrDeselectMeshElementsChangeInput ChangeInput;
-			ChangeInput.MeshElementsToSelect = MeshElementsToSelect;
-			CompoundRevertInput.Subchanges.Add( MoveTemp( FSelectOrDeselectMeshElementsChange( MoveTemp( ChangeInput ) ).Execute( MeshEditorModeProxyObject ) ) );
-
-			TrackUndo( MeshEditorModeProxyObject, MakeUnique<FCompoundChange>( MoveTemp( CompoundRevertInput ) ) );
+			SelectMeshElements( MeshElementsToSelect );
 		}
 
 		return true;
@@ -3235,255 +3194,6 @@ void FMeshEditorMode::UpdateActiveAction( const bool bIsActionFinishing )
 	{
 		bIsMovingSelectedMeshElements = true;
 	}
-	else if( ActiveAction == EMeshEditAction::SplitEdge )
-	{
-		// Split edge
-		static TMap< UEditableMesh*, TArray<FMeshElement> > SelectedMeshesAndEdges;
-		GetSelectedMeshesAndEdges( /* Out */ SelectedMeshesAndEdges );
-
-		if( SelectedMeshesAndEdges.Num() > 0 )
-		{
-			for( auto& MeshAndEdges : SelectedMeshesAndEdges )
-			{
-				UEditableMesh* EditableMesh = MeshAndEdges.Key;
-				verify( !EditableMesh->AnyChangesToUndo() );
-
-				const TArray<FMeshElement>& EdgeElements = MeshAndEdges.Value;
-
-				// Figure out where to split
-				static TArray<float> Splits;
-				FindEdgeSplitUnderInteractor( ActiveActionInteractor, EditableMesh, EdgeElements, /* Out */ Splits );
-
-				// Split the edges
-				if( Splits.Num() > 0 )
-				{
-					// @todo mesheditor edgeloop: Test code
-					if( false )
-					{
-						if( Splits[ 0 ] > 0.25f )
-						{
-							Splits.Insert( FMath::Max( Splits[ 0 ] - 0.2f, 0.0f ), 0 );
-						}
-						if( Splits.Last() < 0.75f )
-						{
-							Splits.Add( FMath::Min( Splits.Last() + 0.2f, 1.0f ) );
-						}
-					}
-
-					for( const FMeshElement& EdgeMeshElement : EdgeElements )
-					{
-						const FEdgeID EdgeID( EdgeMeshElement.ElementAddress.ElementID );
-
-						static TArray<FVertexID> NewVertexIDs;
-						NewVertexIDs.Reset();
-
-						EditableMesh->SplitEdge( EdgeID, Splits, /* Out */ NewVertexIDs );
-
-						// Select all of the new vertices that were created by splitting the edge
-						{
-							bShouldDeselectAllFirst = true;	// Don't keep the edge selected
-							for( const FVertexID NewVertexID : NewVertexIDs )
-							{
-								FMeshElement MeshElementToSelect;
-								{
-									MeshElementToSelect.Component = EdgeMeshElement.Component;
-									MeshElementToSelect.ElementAddress.SubMeshAddress = EdgeMeshElement.ElementAddress.SubMeshAddress;
-									MeshElementToSelect.ElementAddress.ElementType = EEditableMeshElementType::Vertex;
-									MeshElementToSelect.ElementAddress.ElementID = NewVertexID;
-								}
-
-								// Queue selection of this new element.  We don't want it to be part of the current action.
-								MeshElementsToSelect.Add( MeshElementToSelect );
-							}
-						}
-					}
-
-					TrackUndo( EditableMesh, EditableMesh->MakeUndo() );
-				}
-			}
-		}
-	}
-
-	else if( ActiveAction == EMeshEditAction::SplitEdgeAndDragVertex )
-	{
-		if( SplitEdgeMeshesAndEdgesToSplit.Num() > 0 && SplitEdgeSplitList.Num() > 0 )
-		{
-			for( auto& MeshAndEdges : SplitEdgeMeshesAndEdgesToSplit )
-			{
-				UEditableMesh* EditableMesh = MeshAndEdges.Key;
-				verify( !EditableMesh->AnyChangesToUndo() );
-
-				const TArray<FMeshElement>& EdgeElements = MeshAndEdges.Value;
-
-				for( const FMeshElement& EdgeElement : EdgeElements )
-				{
-					const FEdgeID EdgeID( EdgeElement.ElementAddress.ElementID );
-
-					static TArray<FVertexID> NewVertexIDs;
-					NewVertexIDs.Reset();
-
-					EditableMesh->SplitEdge( EdgeID, SplitEdgeSplitList, /* Out */ NewVertexIDs );
-
-					// Select all of the new vertices that were created by splitting the edge
-					{
-						bShouldDeselectAllFirst = true;	// Don't keep the edge selected
-						for( const FVertexID NewVertexID : NewVertexIDs )
-						{
-							FMeshElement MeshElementToSelect;
-							{
-								MeshElementToSelect.Component = EdgeElement.Component;
-								MeshElementToSelect.ElementAddress.SubMeshAddress = EdgeElement.ElementAddress.SubMeshAddress;
-								MeshElementToSelect.ElementAddress.ElementType = EEditableMeshElementType::Vertex;
-								MeshElementToSelect.ElementAddress.ElementID = NewVertexID;
-							}
-
-							// Queue selection of this new element.  We don't want it to be part of the current action.
-							MeshElementsToSelect.Add( MeshElementToSelect );
-
-						}
-
-						if( NewVertexIDs.Num() > 0 )
-						{
-							// Drag the split elements right away!
-							bIsMovingSelectedMeshElements = true;
-						}
-					}
-				}
-
-				TrackUndo( EditableMesh, EditableMesh->MakeUndo() );
-			}
-		}
-	}
-	else if( ActiveAction == EMeshEditAction::ExtendEdge )
-	{
-		// Extend edge
-		static TMap< UEditableMesh*, TArray<FMeshElement> > MeshesWithEdgesToExtend;
-		GetSelectedMeshesAndEdges( /* Out */ MeshesWithEdgesToExtend );
-
-		if( MeshesWithEdgesToExtend.Num() > 0 )
-		{
-			// Allow the user to move the new extended edge around right away!
-			bIsMovingSelectedMeshElements = true;
-
-			for( auto& MeshAndEdges : MeshesWithEdgesToExtend )
-			{
-				UEditableMesh* EditableMesh = MeshAndEdges.Key;
-				const TArray<FMeshElement>& EdgeElementsToExtend = MeshAndEdges.Value;
-
-				static TArray<FEdgeID> EdgeIDsToExtend;
-				EdgeIDsToExtend.Reset();
-
-				for( const FMeshElement& EdgeElementToExtend : EdgeElementsToExtend )
-				{
-					const FEdgeID EdgeID( EdgeElementToExtend.ElementAddress.ElementID );
-
-					EdgeIDsToExtend.Add( EdgeID );
-				}
-
-
-				verify( !EditableMesh->AnyChangesToUndo() );
-
-				{
-					// Extend the edge
-					const bool bWeldNeighbors = true;	// @todo mesheditor urgent extrude: Make optional somehow
-					static TArray<FEdgeID> NewExtendedEdgeIDs;
-					NewExtendedEdgeIDs.Reset();
-					EditableMesh->ExtendEdges( EdgeIDsToExtend, bWeldNeighbors, /* Out */ NewExtendedEdgeIDs );
-
-
-					// Make sure the new edges are selected
-					bShouldDeselectAllFirst = true;	// Deselect the original edges
-					for( int32 NewEdgeNumber = 0; NewEdgeNumber < NewExtendedEdgeIDs.Num(); ++NewEdgeNumber )
-					{
-						const FEdgeID NewExtendedEdgeID = NewExtendedEdgeIDs[ NewEdgeNumber ];
-
-						const FMeshElement& MeshElement = EdgeElementsToExtend[ NewEdgeNumber ];
-
-						FMeshElement NewExtendedEdgeMeshElement;
-						{
-							NewExtendedEdgeMeshElement.Component = MeshElement.Component;
-							NewExtendedEdgeMeshElement.ElementAddress = MeshElement.ElementAddress;
-							NewExtendedEdgeMeshElement.ElementAddress.ElementID = NewExtendedEdgeID;
-						}
-
-						// Queue selection of this new element.  We don't want it to be part of the current action.
-						MeshElementsToSelect.Add( NewExtendedEdgeMeshElement );
-					}
-				}
-
-				TrackUndo( EditableMesh, EditableMesh->MakeUndo() );
-			}
-		}
-	}
-
-	else if( ActiveAction == EMeshEditAction::ExtendVertex )
-	{
-		static TMap< UEditableMesh*, TArray< FMeshElement > > MeshesWithVerticesToExtend;
-		GetSelectedMeshesAndVertices( /* Out */ MeshesWithVerticesToExtend );
-
-		if( MeshesWithVerticesToExtend.Num() > 0 )
-		{
-			// Allow the user to move the new extended vertex around right away!
-			bIsMovingSelectedMeshElements = true;
-
-			for( auto& MeshAndVertices : MeshesWithVerticesToExtend )
-			{
-				UEditableMesh* EditableMesh = MeshAndVertices.Key;
-				const TArray<FMeshElement>& VertexElementsToExtend = MeshAndVertices.Value;
-
-				static TArray<FVertexID> VertexIDsToExtend;
-				VertexIDsToExtend.Reset();
-
-				for( const FMeshElement& VertexElementToExtend : VertexElementsToExtend )
-				{
-					const FVertexID VertexID( VertexElementToExtend.ElementAddress.ElementID );
-
-					VertexIDsToExtend.Add( VertexID );
-				}
-
-
-				verify( !EditableMesh->AnyChangesToUndo() );
-
-				{
-					UPrimitiveComponent* ComponentPtr = VertexElementsToExtend[ 0 ].Component.Get();
-					check( ComponentPtr != nullptr );
-					UPrimitiveComponent& Component = *ComponentPtr;
-					const FMatrix ComponentToWorldMatrix = Component.GetRenderMatrix();
-
-					const FVector ComponentSpaceDragToLocation = ComponentToWorldMatrix.InverseTransformPosition( ActiveActionInteractor->GetInteractorData().LastDragToLocation );
-
-					// Extend the vertex
-					const bool bOnlyExtendClosestEdge = true;	// @todo mesheditor urgent extrude: Make optional somehow
-					static TArray<FVertexID> NewExtendedVertexIDs;
-					NewExtendedVertexIDs.Reset();
-					EditableMesh->ExtendVertices( VertexIDsToExtend, bOnlyExtendClosestEdge, ComponentSpaceDragToLocation, /* Out */ NewExtendedVertexIDs );
-
-
-					// Make sure the new vertices are selected
-					bShouldDeselectAllFirst = true;	// Deselect the original vertices
-					for( int32 NewVertexNumber = 0; NewVertexNumber < NewExtendedVertexIDs.Num(); ++NewVertexNumber )
-					{
-						const FVertexID NewExtendedVertexID = NewExtendedVertexIDs[ NewVertexNumber ];
-
-						const FMeshElement& MeshElement = VertexElementsToExtend[ NewVertexNumber ];
-
-						FMeshElement NewExtendedVertexMeshElement;
-						{
-							NewExtendedVertexMeshElement.Component = MeshElement.Component;
-							NewExtendedVertexMeshElement.ElementAddress = MeshElement.ElementAddress;
-							NewExtendedVertexMeshElement.ElementAddress.ElementID = NewExtendedVertexID;
-						}
-
-						// Queue selection of this new element.  We don't want it to be part of the current action.
-						MeshElementsToSelect.Add( NewExtendedVertexMeshElement );
-					}
-				}
-
-				TrackUndo( EditableMesh, EditableMesh->MakeUndo() );
-			}
-		}
-	}
-
 	else if( ActiveAction == EMeshEditAction::EditVertexCornerSharpness ||
 			 ActiveAction == EMeshEditAction::EditEdgeCreaseSharpness )
 	{
@@ -3889,7 +3599,9 @@ void FMeshEditorMode::UpdateActiveAction( const bool bIsActionFinishing )
 			{
 				if( ActiveAction == CommandCDO->GetCommandName() )
 				{
-					CommandCDO->ApplyDuringDrag( *this, ActiveActionInteractor, /* Out */ bShouldDeselectAllFirst, /* Out */ MeshElementsToSelect );
+					CommandCDO->ApplyDuringDrag( *this, ActiveActionInteractor );
+
+					bIsMovingSelectedMeshElements = CommandCDO->NeedsDraggingInitiated();
 
 					// Should always only be one candidate
 					bFoundValidCommand = true;
@@ -3905,7 +3617,7 @@ void FMeshEditorMode::UpdateActiveAction( const bool bIsActionFinishing )
 	// we'll be dragging any newly-generated geometry from the mesh edit action.  For example, when extending
 	// an edge we want to drag around the newly-created edge, not the edge that was selected before.
 	{
-		if( bShouldDeselectAllFirst )
+		if( bShouldDeselectAllFirst )	// @todo mesheditor extensibility: Remove this stuff ideally
 		{
 			DeselectAllMeshElements();
 		}
@@ -4742,53 +4454,6 @@ void FMeshEditorMode::OnViewportInteractionInputAction( FEditorViewportClient& V
 						const bool bActionNeedsHoverLocation = false;
 						StartAction( EMeshEditAction::Move, ViewportInteractor, bActionNeedsHoverLocation, LOCTEXT( "UndoDragPolygon", "Drag Polygon" ) );
 					}
-					else if( SelectedMeshElementType == EEditableMeshElementType::Edge && EquippedEdgeAction == EMeshEditAction::SplitEdge )
-					{
-						const bool bActionNeedsHoverLocation = true;
-						StartAction( EMeshEditAction::SplitEdge, ViewportInteractor, bActionNeedsHoverLocation, LOCTEXT( "UndoSplitEdge", "Split Edge" ) );
-
-						bOutIsInputCaptured = true;
-					}
-					else if( SelectedMeshElementType == EEditableMeshElementType::Edge && EquippedEdgeAction == EMeshEditAction::SplitEdgeAndDragVertex )
-					{
-						// Figure out what to split
-						GetSelectedMeshesAndEdges( /* Out */ SplitEdgeMeshesAndEdgesToSplit );
-
-						if( SplitEdgeMeshesAndEdgesToSplit.Num() > 0 )
-						{
-							for( auto& MeshAndEdges : SplitEdgeMeshesAndEdgesToSplit )
-							{
-								UEditableMesh* EditableMesh = MeshAndEdges.Key;
-								const TArray<FMeshElement>& EdgeElements = MeshAndEdges.Value;
-
-								// Figure out where to split
-								FindEdgeSplitUnderInteractor( ViewportInteractor, EditableMesh, EdgeElements, /* Out */ SplitEdgeSplitList );
-
-								// Split the edges
-								if( SplitEdgeSplitList.Num() > 0 )
-								{
-									bWantToStartMoving = true;
-									const bool bActionNeedsHoverLocation = false;
-									StartAction( EMeshEditAction::SplitEdgeAndDragVertex, ViewportInteractor, bActionNeedsHoverLocation, LOCTEXT( "UndoSplitEdgeAndDragVertex", "Split Edge And Drag Vertex" ) );
-
-									// We just need to get one split list.  We'll use that for all selected meshes.
-									break;
-								}
-							}
-						}
-					}
-					else if( SelectedMeshElementType == EEditableMeshElementType::Edge && EquippedEdgeAction == EMeshEditAction::ExtendEdge )
-					{
-						bWantToStartMoving = true;
-						const bool bActionNeedsHoverLocation = false;
-						StartAction( EMeshEditAction::ExtendEdge, ViewportInteractor, bActionNeedsHoverLocation, LOCTEXT( "UndoExtendEdge", "Extend Edge" ) );
-					}
-					else if( SelectedMeshElementType == EEditableMeshElementType::Vertex && EquippedVertexAction == EMeshEditAction::ExtendVertex )
-					{
-						bWantToStartMoving = true;
-						const bool bActionNeedsHoverLocation = false;
-						StartAction( EMeshEditAction::ExtendVertex, ViewportInteractor, bActionNeedsHoverLocation, LOCTEXT( "UndoExtendVertex", "Extend Vertex" ) );
-					}
 					else if( SelectedMeshElementType == EEditableMeshElementType::Vertex && EquippedVertexAction == EMeshEditAction::EditVertexCornerSharpness )
 					{
 						EditSharpnessStartLocation = MeshEditorInteractorData.HoverLocation;
@@ -4836,7 +4501,11 @@ void FMeshEditorMode::OnViewportInteractionInputAction( FEditorViewportClient& V
 									{
 										StartAction( EquippedAction, ViewportInteractor, CommandCDO->NeedsHoverLocation(), CommandCDO->GetUndoText() );
 
-										if( !CommandCDO->NeedsDraggingInitiated() )
+										if( CommandCDO->NeedsDraggingInitiated() )
+										{
+											bWantToStartMoving = true;
+										}
+										else
 										{
 											bOutIsInputCaptured = true;
 										}
@@ -5618,19 +5287,6 @@ void FMeshEditorMode::MakeVRRadialMenuActionsMenu(FMenuBuilder& MenuBuilder, TSh
 			NAME_None,
 			EUserInterfaceActionType::CollapsedButton
 			);
-		MenuBuilder.AddMenuEntry(
-			LOCTEXT("ExtendEdge", "Extend"),
-			FText(),
-			FSlateIcon(FMeshEditorStyle::GetStyleSetName(), "MeshEditorMode.EdgeExtend"),
-			FUIAction
-			(
-				FExecuteAction::CreateLambda([this] { SetEquippedAction(EEditableMeshElementType::Edge, EMeshEditAction::ExtendEdge ); }),
-				FCanExecuteAction::CreateSP(this, &FMeshEditorMode::IsMeshElementTypeSelectedOrIsActiveSelectionMode, EEditableMeshElementType::Edge),
-				FIsActionChecked::CreateLambda([this] { return (EquippedEdgeAction == EMeshEditAction::ExtendEdge); })
-				),
-			NAME_None,
-			EUserInterfaceActionType::ToggleButton
-			);
 	}
 	else if (GetMeshElementSelectionMode() == EEditableMeshElementType::Vertex)
 	{
@@ -5671,19 +5327,6 @@ void FMeshEditorMode::MakeVRRadialMenuActionsMenu(FMenuBuilder& MenuBuilder, TSh
 			NAME_None,
 			EUserInterfaceActionType::CollapsedButton
 			);
-		MenuBuilder.AddMenuEntry(
-			LOCTEXT("Extend", "Extend"),
-			FText(),
-			FSlateIcon(FMeshEditorStyle::GetStyleSetName(), "MeshEditorMode.VertexExtend"),
-			FUIAction
-			(
-				FExecuteAction::CreateLambda([this] { SetEquippedAction( EEditableMeshElementType::Vertex, EMeshEditAction::ExtendVertex ); }),
-				FCanExecuteAction::CreateSP(this, &FMeshEditorMode::IsMeshElementTypeSelectedOrIsActiveSelectionMode, EEditableMeshElementType::Vertex),
-				FIsActionChecked::CreateLambda([this] { return (EquippedVertexAction == EMeshEditAction::ExtendVertex); })
-				),
-			NAME_None,
-			EUserInterfaceActionType::ToggleButton
-			);
 		MenuBuilder.AddMenuEntry(
 			LOCTEXT("WeldSelected", "Weld Selected"),
 			FText(),
diff --git a/Engine/Plugins/Editor/MeshEditor/Source/MeshEditor/MeshEditorMode.h b/Engine/Plugins/Editor/MeshEditor/Source/MeshEditor/MeshEditorMode.h
--- a/Engine/Plugins/Editor/MeshEditor/Source/MeshEditor/MeshEditorMode.h
+++ b/Engine/Plugins/Editor/MeshEditor/Source/MeshEditor/MeshEditorMode.h
@@ -585,16 +585,6 @@ protected:
 	TArray<TTuple<TSharedPtr<FUICommandInfo>, FUIAction>> EdgeActions;
 	TArray<TTuple<TSharedPtr<FUICommandInfo>, FUIAction>> PolygonActions;
 
-	//
-	// SplitEdgeAndDragVertex
-	//
-
-	/** When splitting an edge and dragging a vertex, this is the list of edges that will be split */
-	TMap< UEditableMesh*, TArray< FMeshElement > > SplitEdgeMeshesAndEdgesToSplit;
-
-	/** When splitting an edge and dragging a vertex, this is the list of split positions along those edges */
-	TArray<float> SplitEdgeSplitList;
-
 
 	//
 	// EditVertexCornerSharpness/EdgeEdgeCreaseSharpness
diff --git a/Engine/Plugins/Editor/MeshEditor/Source/MeshEditor/Public/IMeshEditorModeEditingContract.h b/Engine/Plugins/Editor/MeshEditor/Source/MeshEditor/Public/IMeshEditorModeEditingContract.h
--- a/Engine/Plugins/Editor/MeshEditor/Source/MeshEditor/Public/IMeshEditorModeEditingContract.h
+++ b/Engine/Plugins/Editor/MeshEditor/Source/MeshEditor/Public/IMeshEditorModeEditingContract.h
@@ -23,18 +23,6 @@ namespace EMeshEditAction
 	/** Moving selected mesh elements (vertices, edges or polygons) */
 	MESHEDITOR_API extern const FName Move;
 
-	/** Split an edge by inserting a vertex.  You can drag to preview where the vertex will be inserted. */
-	MESHEDITOR_API extern const FName SplitEdge;
-
-	/** Splits an edge by inserting a new vertex, then immediately starts dragging that vertex */
-	MESHEDITOR_API extern const FName SplitEdgeAndDragVertex;
-
-	/** Extend an edge by making a copy of it and allowing you to move it around */
-	MESHEDITOR_API extern const FName ExtendEdge;
-
-	/** Extend a vertex by making a copy of it, creating new polygons to join the geometry together */
-	MESHEDITOR_API extern const FName ExtendVertex;
-
 	/** For subdivision meshes, edits how sharp a vertex corner is by dragging in space */
 	MESHEDITOR_API extern const FName EditVertexCornerSharpness;
 
diff --git a/Engine/Plugins/Editor/MeshEditor/Source/MeshEditor/Public/MeshEditorCommands.h b/Engine/Plugins/Editor/MeshEditor/Source/MeshEditor/Public/MeshEditorCommands.h
--- a/Engine/Plugins/Editor/MeshEditor/Source/MeshEditor/Public/MeshEditorCommands.h
+++ b/Engine/Plugins/Editor/MeshEditor/Source/MeshEditor/Public/MeshEditorCommands.h
@@ -32,7 +32,7 @@ public:
 	}
 
 	/** Applies this command every frame while dragging */
-	virtual void ApplyDuringDrag( class IMeshEditorModeEditingContract& MeshEditorMode, class UViewportInteractor* ViewportInteractor, bool& bOutShouldDeselectAllFirst, TArray<FMeshElement>& OutMeshElementsToSelect )
+	virtual void ApplyDuringDrag( class IMeshEditorModeEditingContract& MeshEditorMode, class UViewportInteractor* ViewportInteractor )
 	{
 	}
 
@@ -217,9 +217,6 @@ public:
 	/** Sets the primary action to move vertices */
 	TSharedPtr<FUICommandInfo> MoveVertex;
 
-	/** Sets the primary action to extend vertices */
-	TSharedPtr<FUICommandInfo> ExtendVertex;
-
 	/** Sets the primary action to edit the vertex's corner sharpness */
 	TSharedPtr<FUICommandInfo> EditVertexCornerSharpness;
 
@@ -246,15 +243,6 @@ public:
 	/** Sets the primary action to move edges */
 	TSharedPtr<FUICommandInfo> MoveEdge;
 
-	/** Sets the primary action to split edges */
-	TSharedPtr<FUICommandInfo> SplitEdge;
-
-	/** Sets the primary action to split edges and drag vertices */
-	TSharedPtr<FUICommandInfo> SplitEdgeAndDragVertex;
-
-	/** Sets the primary action to extend edges */
-	TSharedPtr<FUICommandInfo> ExtendEdge;
-
 	/** Sets the primary action to edit the edge's crease sharpness */
 	TSharedPtr<FUICommandInfo> EditEdgeCreaseSharpness;
 
diff --git a/Engine/Plugins/Editor/MeshEditor/Source/PolygonModeling/BevelOrInsetPolygon.cpp b/Engine/Plugins/Editor/MeshEditor/Source/PolygonModeling/BevelOrInsetPolygon.cpp
--- a/Engine/Plugins/Editor/MeshEditor/Source/PolygonModeling/BevelOrInsetPolygon.cpp
+++ b/Engine/Plugins/Editor/MeshEditor/Source/PolygonModeling/BevelOrInsetPolygon.cpp
@@ -127,7 +127,7 @@ namespace BevelOrInsetPolygonHelpers
 	}
 
 
-	static void ApplyDuringDrag( const EBevelOrInset BevelOrInset, IMeshEditorModeEditingContract& MeshEditorMode, class UViewportInteractor* ViewportInteractor, bool& bOutShouldDeselectAllFirst, TArray<FMeshElement>& OutMeshElementsToSelect )
+	static void ApplyDuringDrag( const EBevelOrInset BevelOrInset, IMeshEditorModeEditingContract& MeshEditorMode, class UViewportInteractor* ViewportInteractor )
 	{
 		static TMap< UEditableMesh*, TArray< FMeshElement > > MeshesWithPolygonsToInset;
 		MeshEditorMode.GetSelectedMeshesAndPolygons( /* Out */ MeshesWithPolygonsToInset );
@@ -141,6 +141,11 @@ namespace BevelOrInsetPolygonHelpers
 				UPrimitiveComponent* InsetUsingComponent = InsetUsingPolygonElement.Component.Get();
 				check( InsetUsingComponent != nullptr );
 
+				MeshEditorMode.DeselectAllMeshElements();
+
+				static TArray<FMeshElement> MeshElementsToSelect;
+				MeshElementsToSelect.Reset();
+
 				const UEditableMesh* InsetUsingEditableMesh = MeshEditorMode.FindEditableMesh( *InsetUsingComponent, InsetUsingPolygonElement.ElementAddress.SubMeshAddress );
 				if( ensure( InsetUsingEditableMesh != nullptr ) )
 				{
@@ -217,7 +222,7 @@ namespace BevelOrInsetPolygonHelpers
 										}
 
 										// Queue selection of this new element.  We don't want it to be part of the current action.
-										OutMeshElementsToSelect.Add( PolygonMeshElement );
+										MeshElementsToSelect.Add( PolygonMeshElement );
 									}
 								}
 								else if( ensure( NewSideInsetPolygons.Num() > 0 ) )
@@ -237,7 +242,7 @@ namespace BevelOrInsetPolygonHelpers
 										}
 
 										// Queue selection of this new element.  We don't want it to be part of the current action.
-										OutMeshElementsToSelect.Add( PolygonMeshElement );
+										MeshElementsToSelect.Add( PolygonMeshElement );
 									}
 								}
 							}
@@ -246,6 +251,8 @@ namespace BevelOrInsetPolygonHelpers
 						}
 					}
 				}
+
+				MeshEditorMode.SelectMeshElements( MeshElementsToSelect );
 			}
 		}
 	}
@@ -264,9 +271,9 @@ bool UBevelPolygonCommand::TryStartingToDrag( IMeshEditorModeEditingContract& Me
 }
 
 
-void UBevelPolygonCommand::ApplyDuringDrag( IMeshEditorModeEditingContract& MeshEditorMode, UViewportInteractor* ViewportInteractor, bool& bOutShouldDeselectAllFirst, TArray<FMeshElement>& OutMeshElementsToSelect )
+void UBevelPolygonCommand::ApplyDuringDrag( IMeshEditorModeEditingContract& MeshEditorMode, UViewportInteractor* ViewportInteractor )
 {
-	BevelOrInsetPolygonHelpers::ApplyDuringDrag( BevelOrInsetPolygonHelpers::EBevelOrInset::Bevel, MeshEditorMode, ViewportInteractor, bOutShouldDeselectAllFirst, OutMeshElementsToSelect );
+	BevelOrInsetPolygonHelpers::ApplyDuringDrag( BevelOrInsetPolygonHelpers::EBevelOrInset::Bevel, MeshEditorMode, ViewportInteractor );
 }
 
 
@@ -299,9 +306,9 @@ bool UInsetPolygonCommand::TryStartingToDrag( IMeshEditorModeEditingContract& Me
 }
 
 
-void UInsetPolygonCommand::ApplyDuringDrag( IMeshEditorModeEditingContract& MeshEditorMode, UViewportInteractor* ViewportInteractor, bool& bOutShouldDeselectAllFirst, TArray<FMeshElement>& OutMeshElementsToSelect )
+void UInsetPolygonCommand::ApplyDuringDrag( IMeshEditorModeEditingContract& MeshEditorMode, UViewportInteractor* ViewportInteractor )
 {
-	BevelOrInsetPolygonHelpers::ApplyDuringDrag( BevelOrInsetPolygonHelpers::EBevelOrInset::Inset, MeshEditorMode, ViewportInteractor, bOutShouldDeselectAllFirst, OutMeshElementsToSelect );
+	BevelOrInsetPolygonHelpers::ApplyDuringDrag( BevelOrInsetPolygonHelpers::EBevelOrInset::Inset, MeshEditorMode, ViewportInteractor );
 }
 
 
diff --git a/Engine/Plugins/Editor/MeshEditor/Source/PolygonModeling/BevelOrInsetPolygon.h b/Engine/Plugins/Editor/MeshEditor/Source/PolygonModeling/BevelOrInsetPolygon.h
--- a/Engine/Plugins/Editor/MeshEditor/Source/PolygonModeling/BevelOrInsetPolygon.h
+++ b/Engine/Plugins/Editor/MeshEditor/Source/PolygonModeling/BevelOrInsetPolygon.h
@@ -25,7 +25,7 @@ protected:
 	// Overrides
 	virtual void RegisterUICommand( class FBindingContext* BindingContext ) override;
 	virtual bool TryStartingToDrag( IMeshEditorModeEditingContract& MeshEditorMode, class UViewportInteractor* ViewportInteractor ) override;
-	virtual void ApplyDuringDrag( IMeshEditorModeEditingContract& MeshEditorMode, class UViewportInteractor* ViewportInteractor, bool& bOutShouldDeselectAllFirst, TArray<FMeshElement>& OutMeshElementsToSelect ) override;
+	virtual void ApplyDuringDrag( IMeshEditorModeEditingContract& MeshEditorMode, class UViewportInteractor* ViewportInteractor ) override;
 	virtual void AddToVRRadialMenuActionsMenu( class IMeshEditorModeUIContract& MeshEditorMode, class FMenuBuilder& MenuBuilder, TSharedPtr<FUICommandList> CommandList, const FName TEMPHACK_StyleSetName, class UVREditorMode* VRMode ) override;
 
 };
@@ -50,7 +50,7 @@ protected:
 	// Overrides
 	virtual void RegisterUICommand( class FBindingContext* BindingContext ) override;
 	virtual bool TryStartingToDrag( class IMeshEditorModeEditingContract& MeshEditorMode, class UViewportInteractor* ViewportInteractor ) override;
-	virtual void ApplyDuringDrag( IMeshEditorModeEditingContract& MeshEditorMode, class UViewportInteractor* ViewportInteractor, bool& bOutShouldDeselectAllFirst, TArray<FMeshElement>& OutMeshElementsToSelect ) override;
+	virtual void ApplyDuringDrag( IMeshEditorModeEditingContract& MeshEditorMode, class UViewportInteractor* ViewportInteractor ) override;
 	virtual void AddToVRRadialMenuActionsMenu( class IMeshEditorModeUIContract& MeshEditorMode, class FMenuBuilder& MenuBuilder, TSharedPtr<FUICommandList> CommandList, const FName TEMPHACK_StyleSetName, class UVREditorMode* VRMode ) override;
 
 };
diff --git a/Engine/Plugins/Editor/MeshEditor/Source/PolygonModeling/ExtendEdge.cpp b/Engine/Plugins/Editor/MeshEditor/Source/PolygonModeling/ExtendEdge.cpp
new file mode 100644
--- /dev/null
+++ b/Engine/Plugins/Editor/MeshEditor/Source/PolygonModeling/ExtendEdge.cpp
@@ -0,0 +1,105 @@
+// Copyright 1998-2016 Epic Games, Inc. All Rights Reserved.
+
+#include "ExtendEdge.h"
+#include "IMeshEditorModeEditingContract.h"
+#include "IMeshEditorModeUIContract.h"
+#include "UICommandInfo.h"
+#include "EditableMesh.h"
+#include "MeshElement.h"
+#include "MultiBoxBuilder.h"
+#include "UICommandList.h"
+#include "ViewportInteractor.h"
+
+#define LOCTEXT_NAMESPACE "MeshEditorMode"
+
+
+void UExtendEdgeCommand::RegisterUICommand( FBindingContext* BindingContext )
+{
+	UI_COMMAND_EXT( BindingContext, /* Out */ UICommandInfo, "ExtendEdge", "Extend Edge Mode", "Set the primary action to extend edges.", EUserInterfaceActionType::RadioButton, FInputChord() );
+}
+
+
+void UExtendEdgeCommand::ApplyDuringDrag( IMeshEditorModeEditingContract& MeshEditorMode, UViewportInteractor* ViewportInteractor )
+{
+	// Extend edge
+	static TMap< UEditableMesh*, TArray<FMeshElement> > MeshesWithEdgesToExtend;
+	MeshEditorMode.GetSelectedMeshesAndEdges( /* Out */ MeshesWithEdgesToExtend );
+
+	if( MeshesWithEdgesToExtend.Num() > 0 )
+	{
+		MeshEditorMode.DeselectAllMeshElements();
+
+		static TArray<FMeshElement> MeshElementsToSelect;
+		MeshElementsToSelect.Reset();
+
+		for( auto& MeshAndEdges : MeshesWithEdgesToExtend )
+		{
+			UEditableMesh* EditableMesh = MeshAndEdges.Key;
+			const TArray<FMeshElement>& EdgeElementsToExtend = MeshAndEdges.Value;
+
+			static TArray<FEdgeID> EdgeIDsToExtend;
+			EdgeIDsToExtend.Reset();
+
+			for( const FMeshElement& EdgeElementToExtend : EdgeElementsToExtend )
+			{
+				const FEdgeID EdgeID( EdgeElementToExtend.ElementAddress.ElementID );
+
+				EdgeIDsToExtend.Add( EdgeID );
+			}
+
+
+			verify( !EditableMesh->AnyChangesToUndo() );
+
+			{
+				// Extend the edge
+				const bool bWeldNeighbors = true;	// @todo mesheditor urgent extrude: Make optional somehow
+				static TArray<FEdgeID> NewExtendedEdgeIDs;
+				NewExtendedEdgeIDs.Reset();
+				EditableMesh->ExtendEdges( EdgeIDsToExtend, bWeldNeighbors, /* Out */ NewExtendedEdgeIDs );
+
+				// Make sure the new edges are selected
+				for( int32 NewEdgeNumber = 0; NewEdgeNumber < NewExtendedEdgeIDs.Num(); ++NewEdgeNumber )
+				{
+					const FEdgeID NewExtendedEdgeID = NewExtendedEdgeIDs[ NewEdgeNumber ];
+
+					const FMeshElement& MeshElement = EdgeElementsToExtend[ NewEdgeNumber ];
+
+					FMeshElement NewExtendedEdgeMeshElement;
+					{
+						NewExtendedEdgeMeshElement.Component = MeshElement.Component;
+						NewExtendedEdgeMeshElement.ElementAddress = MeshElement.ElementAddress;
+						NewExtendedEdgeMeshElement.ElementAddress.ElementID = NewExtendedEdgeID;
+					}
+
+					// Queue selection of this new element.  We don't want it to be part of the current action.
+					MeshElementsToSelect.Add( NewExtendedEdgeMeshElement );
+				}
+			}
+
+			MeshEditorMode.TrackUndo( EditableMesh, EditableMesh->MakeUndo() );
+		}
+
+		MeshEditorMode.SelectMeshElements( MeshElementsToSelect );
+	}
+}
+
+
+void UExtendEdgeCommand::AddToVRRadialMenuActionsMenu( IMeshEditorModeUIContract& MeshEditorMode, FMenuBuilder& MenuBuilder, TSharedPtr<FUICommandList> CommandList, const FName TEMPHACK_StyleSetName, class UVREditorMode* VRMode )
+{
+	if( MeshEditorMode.GetMeshElementSelectionMode() == EEditableMeshElementType::Edge )
+	{
+		MenuBuilder.AddMenuEntry(
+			LOCTEXT( "VRExtendEdge", "Extend" ),
+			FText(),
+			FSlateIcon( TEMPHACK_StyleSetName, "MeshEditorMode.EdgeExtend" ),
+			MakeUIAction( MeshEditorMode ),
+			NAME_None,
+			EUserInterfaceActionType::ToggleButton
+		);
+	}
+}
+
+
+
+#undef LOCTEXT_NAMESPACE
+
diff --git a/Engine/Plugins/Editor/MeshEditor/Source/PolygonModeling/ExtendEdge.h b/Engine/Plugins/Editor/MeshEditor/Source/PolygonModeling/ExtendEdge.h
new file mode 100644
--- /dev/null
+++ b/Engine/Plugins/Editor/MeshEditor/Source/PolygonModeling/ExtendEdge.h
@@ -0,0 +1,32 @@
+// Copyright 1998-2016 Epic Games, Inc. All Rights Reserved.
+
+#pragma once
+
+#include "MeshEditorCommands.h"
+#include "ExtendEdge.generated.h"
+
+
+/** Extends an edge by making a copy of it and allowing you to move it around */
+UCLASS()
+class UExtendEdgeCommand : public UMeshEditorEdgeCommand
+{
+	GENERATED_BODY()
+
+protected:
+
+	UExtendEdgeCommand()
+	{
+		bIsMode = true;
+		UndoText = NSLOCTEXT( "MeshEditor", "UndoExtendEdge", "Extend Edge" );
+		bNeedsHoverLocation = false;
+		bNeedsDraggingInitiated = true;
+	}
+
+	// Overrides
+	virtual void RegisterUICommand( class FBindingContext* BindingContext ) override;
+	virtual void ApplyDuringDrag( IMeshEditorModeEditingContract& MeshEditorMode, class UViewportInteractor* ViewportInteractor ) override;
+	virtual void AddToVRRadialMenuActionsMenu( class IMeshEditorModeUIContract& MeshEditorMode, class FMenuBuilder& MenuBuilder, TSharedPtr<FUICommandList> CommandList, const FName TEMPHACK_StyleSetName, class UVREditorMode* VRMode ) override;
+
+protected:
+
+};
diff --git a/Engine/Plugins/Editor/MeshEditor/Source/PolygonModeling/ExtendVertex.cpp b/Engine/Plugins/Editor/MeshEditor/Source/PolygonModeling/ExtendVertex.cpp
new file mode 100644
--- /dev/null
+++ b/Engine/Plugins/Editor/MeshEditor/Source/PolygonModeling/ExtendVertex.cpp
@@ -0,0 +1,112 @@
+// Copyright 1998-2016 Epic Games, Inc. All Rights Reserved.
+
+#include "ExtendVertex.h"
+#include "IMeshEditorModeEditingContract.h"
+#include "IMeshEditorModeUIContract.h"
+#include "UICommandInfo.h"
+#include "EditableMesh.h"
+#include "MeshElement.h"
+#include "MultiBoxBuilder.h"
+#include "UICommandList.h"
+#include "ViewportInteractor.h"
+
+#define LOCTEXT_NAMESPACE "MeshEditorMode"
+
+
+void UExtendVertexCommand::RegisterUICommand( FBindingContext* BindingContext )
+{
+	UI_COMMAND_EXT( BindingContext, /* Out */ UICommandInfo, "ExtendVertex", "Extend Vertex Mode", "Set the primary action to extend vertices.", EUserInterfaceActionType::RadioButton, FInputChord() );
+}
+
+
+void UExtendVertexCommand::ApplyDuringDrag( IMeshEditorModeEditingContract& MeshEditorMode, UViewportInteractor* ViewportInteractor )
+{
+	static TMap< UEditableMesh*, TArray< FMeshElement > > MeshesWithVerticesToExtend;
+	MeshEditorMode.GetSelectedMeshesAndVertices( /* Out */ MeshesWithVerticesToExtend );
+
+	if( MeshesWithVerticesToExtend.Num() > 0 )
+	{
+		MeshEditorMode.DeselectAllMeshElements();
+
+		static TArray<FMeshElement> MeshElementsToSelect;
+		MeshElementsToSelect.Reset();
+
+		for( auto& MeshAndVertices : MeshesWithVerticesToExtend )
+		{
+			UEditableMesh* EditableMesh = MeshAndVertices.Key;
+			const TArray<FMeshElement>& VertexElementsToExtend = MeshAndVertices.Value;
+
+			static TArray<FVertexID> VertexIDsToExtend;
+			VertexIDsToExtend.Reset();
+
+			for( const FMeshElement& VertexElementToExtend : VertexElementsToExtend )
+			{
+				const FVertexID VertexID( VertexElementToExtend.ElementAddress.ElementID );
+
+				VertexIDsToExtend.Add( VertexID );
+			}
+
+
+			verify( !EditableMesh->AnyChangesToUndo() );
+
+			{
+				UPrimitiveComponent* ComponentPtr = VertexElementsToExtend[ 0 ].Component.Get();
+				check( ComponentPtr != nullptr );
+				UPrimitiveComponent& Component = *ComponentPtr;
+				const FMatrix ComponentToWorldMatrix = Component.GetRenderMatrix();
+
+				const FVector ComponentSpaceDragToLocation = ComponentToWorldMatrix.InverseTransformPosition( ViewportInteractor->GetInteractorData().LastDragToLocation );
+
+				// Extend the vertex
+				const bool bOnlyExtendClosestEdge = true;	// @todo mesheditor urgent extrude: Make optional somehow
+				static TArray<FVertexID> NewExtendedVertexIDs;
+				NewExtendedVertexIDs.Reset();
+				EditableMesh->ExtendVertices( VertexIDsToExtend, bOnlyExtendClosestEdge, ComponentSpaceDragToLocation, /* Out */ NewExtendedVertexIDs );
+
+
+				// Make sure the new vertices are selected
+				for( int32 NewVertexNumber = 0; NewVertexNumber < NewExtendedVertexIDs.Num(); ++NewVertexNumber )
+				{
+					const FVertexID NewExtendedVertexID = NewExtendedVertexIDs[ NewVertexNumber ];
+
+					const FMeshElement& MeshElement = VertexElementsToExtend[ NewVertexNumber ];
+
+					FMeshElement NewExtendedVertexMeshElement;
+					{
+						NewExtendedVertexMeshElement.Component = MeshElement.Component;
+						NewExtendedVertexMeshElement.ElementAddress = MeshElement.ElementAddress;
+						NewExtendedVertexMeshElement.ElementAddress.ElementID = NewExtendedVertexID;
+					}
+
+					// Queue selection of this new element.  We don't want it to be part of the current action.
+					MeshElementsToSelect.Add( NewExtendedVertexMeshElement );
+				}
+			}
+
+			MeshEditorMode.TrackUndo( EditableMesh, EditableMesh->MakeUndo() );
+		}
+
+		MeshEditorMode.SelectMeshElements( MeshElementsToSelect );
+	}
+}
+
+
+void UExtendVertexCommand::AddToVRRadialMenuActionsMenu( IMeshEditorModeUIContract& MeshEditorMode, FMenuBuilder& MenuBuilder, TSharedPtr<FUICommandList> CommandList, const FName TEMPHACK_StyleSetName, class UVREditorMode* VRMode )
+{
+	if( MeshEditorMode.GetMeshElementSelectionMode() == EEditableMeshElementType::Vertex )
+	{
+		MenuBuilder.AddMenuEntry(
+			LOCTEXT( "VRExtendVertex", "Extend" ),
+			FText(),
+			FSlateIcon( TEMPHACK_StyleSetName, "MeshEditorMode.VertexExtend" ),
+			MakeUIAction( MeshEditorMode ),
+			NAME_None,
+			EUserInterfaceActionType::ToggleButton
+		);
+	}
+}
+
+
+
+#undef LOCTEXT_NAMESPACE
+
diff --git a/Engine/Plugins/Editor/MeshEditor/Source/PolygonModeling/ExtendVertex.h b/Engine/Plugins/Editor/MeshEditor/Source/PolygonModeling/ExtendVertex.h
new file mode 100644
--- /dev/null
+++ b/Engine/Plugins/Editor/MeshEditor/Source/PolygonModeling/ExtendVertex.h
@@ -0,0 +1,32 @@
+// Copyright 1998-2016 Epic Games, Inc. All Rights Reserved.
+
+#pragma once
+
+#include "MeshEditorCommands.h"
+#include "ExtendVertex.generated.h"
+
+
+/** Extend a vertex by making a copy of it, creating new polygons to join the geometry together */
+UCLASS()
+class UExtendVertexCommand : public UMeshEditorVertexCommand
+{
+	GENERATED_BODY()
+
+protected:
+
+	UExtendVertexCommand()
+	{
+		bIsMode = true;
+		UndoText = NSLOCTEXT( "MeshEditor", "UndoExtendVertex", "Extend Vertex" );
+		bNeedsHoverLocation = false;
+		bNeedsDraggingInitiated = true;
+	}
+
+	// Overrides
+	virtual void RegisterUICommand( class FBindingContext* BindingContext ) override;
+	virtual void ApplyDuringDrag( IMeshEditorModeEditingContract& MeshEditorMode, class UViewportInteractor* ViewportInteractor ) override;
+	virtual void AddToVRRadialMenuActionsMenu( class IMeshEditorModeUIContract& MeshEditorMode, class FMenuBuilder& MenuBuilder, TSharedPtr<FUICommandList> CommandList, const FName TEMPHACK_StyleSetName, class UVREditorMode* VRMode ) override;
+
+protected:
+
+};
diff --git a/Engine/Plugins/Editor/MeshEditor/Source/PolygonModeling/ExtrudePolygon.cpp b/Engine/Plugins/Editor/MeshEditor/Source/PolygonModeling/ExtrudePolygon.cpp
--- a/Engine/Plugins/Editor/MeshEditor/Source/PolygonModeling/ExtrudePolygon.cpp
+++ b/Engine/Plugins/Editor/MeshEditor/Source/PolygonModeling/ExtrudePolygon.cpp
@@ -115,7 +115,7 @@ bool UExtrudePolygonCommand::TryStartingToDrag( IMeshEditorModeEditingContract&
 }
 
 
-void UExtrudePolygonCommand::ApplyDuringDrag( IMeshEditorModeEditingContract& MeshEditorMode, UViewportInteractor* ViewportInteractor, bool& bOutShouldDeselectAllFirst, TArray<FMeshElement>& OutMeshElementsToSelect )
+void UExtrudePolygonCommand::ApplyDuringDrag( IMeshEditorModeEditingContract& MeshEditorMode, UViewportInteractor* ViewportInteractor )
 {
 	static TMap< UEditableMesh*, TArray< FMeshElement > > MeshesWithPolygonsToExtrude;
 	MeshEditorMode.GetSelectedMeshesAndPolygons( /* Out */ MeshesWithPolygonsToExtrude );
@@ -125,6 +125,9 @@ void UExtrudePolygonCommand::ApplyDuringDrag( IMeshEditorModeEditingContract& Me
 		// Deselect the mesh elements before we delete them.  This will make sure they become selected again after undo.
 		MeshEditorMode.DeselectMeshElements( MeshesWithPolygonsToExtrude );
 
+		static TArray<FMeshElement> MeshElementsToSelect;
+		MeshElementsToSelect.Reset();
+
 		for( auto& MeshAndPolygons : MeshesWithPolygonsToExtrude )
 		{
 			UEditableMesh* EditableMesh = MeshAndPolygons.Key;
@@ -185,12 +188,14 @@ void UExtrudePolygonCommand::ApplyDuringDrag( IMeshEditorModeEditingContract& Me
 					}
 
 					// Queue selection of this new element.  We don't want it to be part of the current action.
-					OutMeshElementsToSelect.Add( NewExtrudedPolygonMeshElement );
+					MeshElementsToSelect.Add( NewExtrudedPolygonMeshElement );
 				}
 			}
 
 			MeshEditorMode.TrackUndo( EditableMesh, EditableMesh->MakeUndo() );
 		}
+
+		MeshEditorMode.SelectMeshElements( MeshElementsToSelect );
 	}
 }
 
diff --git a/Engine/Plugins/Editor/MeshEditor/Source/PolygonModeling/ExtrudePolygon.h b/Engine/Plugins/Editor/MeshEditor/Source/PolygonModeling/ExtrudePolygon.h
--- a/Engine/Plugins/Editor/MeshEditor/Source/PolygonModeling/ExtrudePolygon.h
+++ b/Engine/Plugins/Editor/MeshEditor/Source/PolygonModeling/ExtrudePolygon.h
@@ -37,7 +37,7 @@ protected:
 	// Overrides
 	virtual void RegisterUICommand( class FBindingContext* BindingContext ) override;
 	virtual bool TryStartingToDrag( IMeshEditorModeEditingContract& MeshEditorMode, class UViewportInteractor* ViewportInteractor ) override;
-	virtual void ApplyDuringDrag( IMeshEditorModeEditingContract& MeshEditorMode, class UViewportInteractor* ViewportInteractor, bool& bOutShouldDeselectAllFirst, TArray<FMeshElement>& OutMeshElementsToSelect ) override;
+	virtual void ApplyDuringDrag( IMeshEditorModeEditingContract& MeshEditorMode, class UViewportInteractor* ViewportInteractor ) override;
 	virtual void AddToVRRadialMenuActionsMenu( class IMeshEditorModeUIContract& MeshEditorMode, class FMenuBuilder& MenuBuilder, TSharedPtr<FUICommandList> CommandList, const FName TEMPHACK_StyleSetName, class UVREditorMode* VRMode ) override;
 
 protected:
diff --git a/Engine/Plugins/Editor/MeshEditor/Source/PolygonModeling/InsertEdgeLoop.cpp b/Engine/Plugins/Editor/MeshEditor/Source/PolygonModeling/InsertEdgeLoop.cpp
--- a/Engine/Plugins/Editor/MeshEditor/Source/PolygonModeling/InsertEdgeLoop.cpp
+++ b/Engine/Plugins/Editor/MeshEditor/Source/PolygonModeling/InsertEdgeLoop.cpp
@@ -19,7 +19,7 @@ void UInsertEdgeLoopCommand::RegisterUICommand( FBindingContext* BindingContext
 }
 
 
-void UInsertEdgeLoopCommand::ApplyDuringDrag( IMeshEditorModeEditingContract& MeshEditorMode, UViewportInteractor* ViewportInteractor, bool& bOutShouldDeselectAllFirst, TArray<FMeshElement>& OutMeshElementsToSelect )
+void UInsertEdgeLoopCommand::ApplyDuringDrag( IMeshEditorModeEditingContract& MeshEditorMode, UViewportInteractor* ViewportInteractor )
 {
 	// Insert edge loop
 	static TMap< UEditableMesh*, TArray< FMeshElement > > SelectedMeshesAndEdges;
@@ -29,8 +29,10 @@ void UInsertEdgeLoopCommand::ApplyDuringDrag( IMeshEditorModeEditingContract& Me
 	{
 		// Deselect the edges first, since they'll be deleted or split up while inserting the edge loop,
 		// and we want them to be re-selected after undo
-		MeshEditorMode.DeselectMeshElements( SelectedMeshesAndEdges );
+		MeshEditorMode.DeselectAllMeshElements();
 
+		static TArray<FMeshElement> MeshElementsToSelect;
+		MeshElementsToSelect.Reset();
 
 		for( auto& MeshAndEdges : SelectedMeshesAndEdges )
 		{
@@ -71,8 +73,6 @@ void UInsertEdgeLoopCommand::ApplyDuringDrag( IMeshEditorModeEditingContract& Me
 					// Select all of the new edges that were created by inserting the loop
 					if( NewEdgeIDs.Num() > 0 )
 					{
-						bOutShouldDeselectAllFirst = true;	// Don't keep the edge selected
-
 						for( const FEdgeID NewEdgeID : NewEdgeIDs )
 						{
 							FMeshElement MeshElementToSelect;
@@ -84,7 +84,7 @@ void UInsertEdgeLoopCommand::ApplyDuringDrag( IMeshEditorModeEditingContract& Me
 							}
 
 							// Queue selection of this new element.  We don't want it to be part of the current action.
-							OutMeshElementsToSelect.Add( MeshElementToSelect );
+							MeshElementsToSelect.Add( MeshElementToSelect );
 						}
 					}
 				}
@@ -92,6 +92,8 @@ void UInsertEdgeLoopCommand::ApplyDuringDrag( IMeshEditorModeEditingContract& Me
 				MeshEditorMode.TrackUndo( EditableMesh, EditableMesh->MakeUndo() );
 			}
 		}
+
+		MeshEditorMode.SelectMeshElements( MeshElementsToSelect );
 	}
 }
 
diff --git a/Engine/Plugins/Editor/MeshEditor/Source/PolygonModeling/InsertEdgeLoop.h b/Engine/Plugins/Editor/MeshEditor/Source/PolygonModeling/InsertEdgeLoop.h
--- a/Engine/Plugins/Editor/MeshEditor/Source/PolygonModeling/InsertEdgeLoop.h
+++ b/Engine/Plugins/Editor/MeshEditor/Source/PolygonModeling/InsertEdgeLoop.h
@@ -24,7 +24,7 @@ protected:
 
 	// Overrides
 	virtual void RegisterUICommand( class FBindingContext* BindingContext ) override;
-	virtual void ApplyDuringDrag( IMeshEditorModeEditingContract& MeshEditorMode, class UViewportInteractor* ViewportInteractor, bool& bOutShouldDeselectAllFirst, TArray<FMeshElement>& OutMeshElementsToSelect ) override;
+	virtual void ApplyDuringDrag( IMeshEditorModeEditingContract& MeshEditorMode, class UViewportInteractor* ViewportInteractor ) override;
 	virtual void AddToVRRadialMenuActionsMenu( class IMeshEditorModeUIContract& MeshEditorMode, class FMenuBuilder& MenuBuilder, TSharedPtr<FUICommandList> CommandList, const FName TEMPHACK_StyleSetName, class UVREditorMode* VRMode ) override;
 
 };
diff --git a/Engine/Plugins/Editor/MeshEditor/Source/PolygonModeling/SplitEdge.cpp b/Engine/Plugins/Editor/MeshEditor/Source/PolygonModeling/SplitEdge.cpp
new file mode 100644
--- /dev/null
+++ b/Engine/Plugins/Editor/MeshEditor/Source/PolygonModeling/SplitEdge.cpp
@@ -0,0 +1,104 @@
+// Copyright 1998-2016 Epic Games, Inc. All Rights Reserved.
+
+#include "SplitEdge.h"
+#include "IMeshEditorModeEditingContract.h"
+#include "IMeshEditorModeUIContract.h"
+#include "UICommandInfo.h"
+#include "EditableMesh.h"
+#include "MeshElement.h"
+#include "MultiBoxBuilder.h"
+#include "UICommandList.h"
+#include "ViewportInteractor.h"
+
+#define LOCTEXT_NAMESPACE "MeshEditorMode"
+
+
+void USplitEdgeCommand::RegisterUICommand( FBindingContext* BindingContext )
+{
+	UI_COMMAND_EXT( BindingContext, /* Out */ UICommandInfo, "SplitEdge", "Split Edge Mode", "Set the primary action to split edges.", EUserInterfaceActionType::RadioButton, FInputChord() );
+}
+
+
+bool USplitEdgeCommand::TryStartingToDrag( IMeshEditorModeEditingContract& MeshEditorMode, UViewportInteractor* ViewportInteractor )
+{
+	bool bHaveEdge = false;
+
+	// Figure out what to split
+	MeshEditorMode.GetSelectedMeshesAndEdges( /* Out */ SplitEdgeMeshesAndEdgesToSplit );
+
+	if( SplitEdgeMeshesAndEdgesToSplit.Num() > 0 )
+	{
+		for( auto& MeshAndEdges : SplitEdgeMeshesAndEdgesToSplit )
+		{
+			UEditableMesh* EditableMesh = MeshAndEdges.Key;
+			const TArray<FMeshElement>& EdgeElements = MeshAndEdges.Value;
+
+			// Figure out where to split
+			MeshEditorMode.FindEdgeSplitUnderInteractor( ViewportInteractor, EditableMesh, EdgeElements, /* Out */ SplitEdgeSplitList );
+
+			// Split the edges
+			if( SplitEdgeSplitList.Num() > 0 )
+			{
+				bHaveEdge = true;
+			}
+		}
+	}
+
+	return bHaveEdge;
+}
+
+
+void USplitEdgeCommand::ApplyDuringDrag( IMeshEditorModeEditingContract& MeshEditorMode, UViewportInteractor* ViewportInteractor )
+{
+	if( SplitEdgeMeshesAndEdgesToSplit.Num() > 0 && SplitEdgeSplitList.Num() > 0 )
+	{
+		MeshEditorMode.DeselectAllMeshElements();
+
+		static TArray<FMeshElement> MeshElementsToSelect;
+		MeshElementsToSelect.Reset();
+
+		for( auto& MeshAndEdges : SplitEdgeMeshesAndEdgesToSplit )
+		{
+			UEditableMesh* EditableMesh = MeshAndEdges.Key;
+			verify( !EditableMesh->AnyChangesToUndo() );
+
+			const TArray<FMeshElement>& EdgeElements = MeshAndEdges.Value;
+
+			for( const FMeshElement& EdgeElement : EdgeElements )
+			{
+				const FEdgeID EdgeID( EdgeElement.ElementAddress.ElementID );
+
+				static TArray<FVertexID> NewVertexIDs;
+				NewVertexIDs.Reset();
+
+				EditableMesh->SplitEdge( EdgeID, SplitEdgeSplitList, /* Out */ NewVertexIDs );
+
+				// Select all of the new vertices that were created by splitting the edge
+				{
+					for( const FVertexID NewVertexID : NewVertexIDs )
+					{
+						FMeshElement MeshElementToSelect;
+						{
+							MeshElementToSelect.Component = EdgeElement.Component;
+							MeshElementToSelect.ElementAddress.SubMeshAddress = EdgeElement.ElementAddress.SubMeshAddress;
+							MeshElementToSelect.ElementAddress.ElementType = EEditableMeshElementType::Vertex;
+							MeshElementToSelect.ElementAddress.ElementID = NewVertexID;
+						}
+
+						// Queue selection of this new element.  We don't want it to be part of the current action.
+						MeshElementsToSelect.Add( MeshElementToSelect );
+					}
+				}
+			}
+
+			MeshEditorMode.TrackUndo( EditableMesh, EditableMesh->MakeUndo() );
+		}
+
+		MeshEditorMode.SelectMeshElements( MeshElementsToSelect );
+	}
+}
+
+
+
+#undef LOCTEXT_NAMESPACE
+
diff --git a/Engine/Plugins/Editor/MeshEditor/Source/PolygonModeling/SplitEdge.h b/Engine/Plugins/Editor/MeshEditor/Source/PolygonModeling/SplitEdge.h
new file mode 100644
--- /dev/null
+++ b/Engine/Plugins/Editor/MeshEditor/Source/PolygonModeling/SplitEdge.h
@@ -0,0 +1,40 @@
+// Copyright 1998-2016 Epic Games, Inc. All Rights Reserved.
+
+#pragma once
+
+#include "MeshEditorCommands.h"
+#include "SplitEdge.generated.h"
+
+
+/** Inserts a vertex along an edge and allows the user to move it around */
+UCLASS()
+class USplitEdgeCommand : public UMeshEditorEdgeCommand
+{
+	GENERATED_BODY()
+
+protected:
+
+	USplitEdgeCommand()
+	  : SplitEdgeMeshesAndEdgesToSplit(),
+	    SplitEdgeSplitList()
+	{
+		bIsMode = true;
+		UndoText = NSLOCTEXT( "MeshEditor", "UndoSplitEdge", "Split Edge" );
+		bNeedsHoverLocation = false;
+		bNeedsDraggingInitiated = true;
+	}
+
+	// Overrides
+	virtual void RegisterUICommand( class FBindingContext* BindingContext ) override;
+	virtual bool TryStartingToDrag( IMeshEditorModeEditingContract& MeshEditorMode, class UViewportInteractor* ViewportInteractor ) override;
+	virtual void ApplyDuringDrag( IMeshEditorModeEditingContract& MeshEditorMode, class UViewportInteractor* ViewportInteractor ) override;
+
+protected:
+
+	/** When splitting an edge and dragging a vertex, this is the list of edges that will be split */
+	TMap< class UEditableMesh*, TArray< FMeshElement > > SplitEdgeMeshesAndEdgesToSplit;
+
+	/** When splitting an edge and dragging a vertex, this is the list of split positions along those edges */
+	TArray<float> SplitEdgeSplitList;
+
+};
