Mesh Editor extensibility (phase 6)
- Edit Edge Crease Sharpness and Edit Vertex Corner Sharpness have been migrated to a new module
- Eliminated most use of EMeshEditAction enums; We just use FNames now.

diff --git a/Engine/Plugins/Editor/MeshEditor/Source/MeshEditor/MeshEditorCommands.cpp b/Engine/Plugins/Editor/MeshEditor/Source/MeshEditor/MeshEditorCommands.cpp
--- a/Engine/Plugins/Editor/MeshEditor/Source/MeshEditor/MeshEditorCommands.cpp
+++ b/Engine/Plugins/Editor/MeshEditor/Source/MeshEditor/MeshEditorCommands.cpp
@@ -15,12 +15,12 @@ FUIAction UMeshEditorCommand::MakeUIAction( IMeshEditorModeUIContract& MeshEdito
 	FUIAction UIAction;
 	if( IsMode() )
 	{
-		const EMeshEditAction::Type MeshEditAction = GetCommandName();
+		const FName CommandName = GetCommandName();
 
 		UIAction = FUIAction(
-			FExecuteAction::CreateLambda( [&MeshEditorMode, ElementType, MeshEditAction] { MeshEditorMode.SetEquippedAction( ElementType, MeshEditAction ); } ),
+			FExecuteAction::CreateLambda( [&MeshEditorMode, ElementType, CommandName] { MeshEditorMode.SetEquippedAction( ElementType, CommandName ); } ),
 			FCanExecuteAction::CreateLambda( [&MeshEditorMode, ElementType] { return MeshEditorMode.IsMeshElementTypeSelectedOrIsActiveSelectionMode( ElementType ); } ),
-			FIsActionChecked::CreateLambda( [&MeshEditorMode, ElementType, MeshEditAction] { return ( MeshEditorMode.GetEquippedAction( ElementType ) == MeshEditAction ); } )
+			FIsActionChecked::CreateLambda( [&MeshEditorMode, ElementType, CommandName] { return ( MeshEditorMode.GetEquippedAction( ElementType ) == CommandName ); } )
 		);
 	}
 	else
@@ -81,7 +81,6 @@ FMeshEditorVertexCommands::FMeshEditorVertexCommands()
 void FMeshEditorVertexCommands::RegisterCommands()
 {
 	UI_COMMAND(MoveVertex, "Move Vertex Mode", "Set the primary action to move vertices.", EUserInterfaceActionType::RadioButton, FInputChord(EKeys::F1));
-	UI_COMMAND(EditVertexCornerSharpness, "Edit Vertex Corner Sharpness Mode", "Set the primary action to edit the vertex's subdivision corner sharpness amount.", EUserInterfaceActionType::RadioButton, FInputChord(EKeys::F6));
 
 	UI_COMMAND(RemoveVertex, "Remove Vertex", "Remove the selected vertex if possible.", EUserInterfaceActionType::Button, FInputChord(EKeys::BackSpace));
 	UI_COMMAND(WeldVertices, "Weld Vertices", "Weld the selected vertices, keeping the first selected vertex.", EUserInterfaceActionType::Button, FInputChord());
@@ -110,7 +109,6 @@ FMeshEditorEdgeCommands::FMeshEditorEdgeCommands()
 void FMeshEditorEdgeCommands::RegisterCommands()
 {
 	UI_COMMAND(MoveEdge, "Move Edge Mode", "Set the primary action to move edges.", EUserInterfaceActionType::RadioButton, FInputChord(EKeys::F1));
-	UI_COMMAND(EditEdgeCreaseSharpness, "Edit Edge Crease Sharpness Mode", "Set the primary action to edit the edge's crease sharpness.", EUserInterfaceActionType::RadioButton, FInputChord(EKeys::F6));
 
 	UI_COMMAND(RemoveEdge, "Remove Edge", "Remove the selected edge if possible.", EUserInterfaceActionType::Button, FInputChord(EKeys::BackSpace));
 	UI_COMMAND(SoftenEdge, "Soften Edge", "Make selected edge soft.", EUserInterfaceActionType::Button, FInputChord(EKeys::H, EModifierKey::Shift));
diff --git a/Engine/Plugins/Editor/MeshEditor/Source/MeshEditor/MeshEditorMode.cpp b/Engine/Plugins/Editor/MeshEditor/Source/MeshEditor/MeshEditorMode.cpp
--- a/Engine/Plugins/Editor/MeshEditor/Source/MeshEditor/MeshEditorMode.cpp
+++ b/Engine/Plugins/Editor/MeshEditor/Source/MeshEditor/MeshEditorMode.cpp
@@ -43,16 +43,19 @@
 // @todo mesheditor extensibility: This should probably be removed after we've evicted all current mesh editing actions to another module
 namespace EMeshEditAction
 {
-	const Type None( "None" );
-	const Type SelectByPainting( "SelectByPainting" );
-	const Type MoveUsingGizmo( "MoveUsingGizmo" );
-	const Type Move( "Move" );
-	const Type EditVertexCornerSharpness( "EditVertexCornerSharpness" );
-	const Type EditEdgeCreaseSharpness( "EditEdgeCreaseSharpness" );
-	const Type DrawVertices( "DrawVertices" );
+	/** Selecting mesh elements by 'painting' over multiple elements */
+	const FName SelectByPainting( "SelectByPainting" );
+
+	/** Moving elements using a transform gizmo */
+	const FName MoveUsingGizmo( "MoveUsingGizmo" );
+
+	/** Moving selected mesh elements (vertices, edges or polygons) */
+	const FName Move( "Move" );
+
+	/** Freehand vertex drawing */
+	const FName DrawVertices( "DrawVertices" );
 }
 
-// @todo mesheditor urgent: Texture coords are getting trashed all the time (only for meshes with 1 UV set though?)
 
 namespace MeshEd
 {
@@ -380,7 +383,7 @@ FMeshEditorMode::FMeshEditorMode()
 	  EquippedVertexAction( EMeshEditAction::Move ),
 	  EquippedEdgeAction( EMeshEditAction::Move ),
 	  EquippedPolygonAction( EMeshEditAction::Move ),
-	  ActiveAction( EMeshEditAction::None ),
+	  ActiveAction( NAME_None ),
 	  bIsCapturingUndoForPreview( false ),
 	  PreviewRevertChanges(),
 	  ActiveActionModifiedMeshes(),
@@ -388,7 +391,6 @@ FMeshEditorMode::FMeshEditorMode()
 	  ActiveActionInteractor( nullptr ),
 	  bActiveActionNeedsHoverLocation( false ),
 	  bIsFirstActiveActionUpdate( false ),
-	  EditSharpnessStartLocation( FVector::ZeroVector ),
 	  SelectingByPaintingRevertChangeInput( nullptr ),
 	  bShowVertexNormals( false ),
 	  bMarqueeSelectTransactionActive( false ),
@@ -484,7 +486,7 @@ void FMeshEditorMode::RemoveEditableMeshReferences()
 	}
 }
 
-void FMeshEditorMode::PlayStartActionSound(EMeshEditAction::Type NewAction, UViewportInteractor* ActionInteractor)
+void FMeshEditorMode::PlayStartActionSound( FName NewAction, UViewportInteractor* ActionInteractor)
 {
 	if (ActionInteractor != nullptr)
 	{
@@ -498,7 +500,7 @@ void FMeshEditorMode::PlayStartActionSound(EMeshEditAction::Type NewAction, UVie
 	}
 }
 
-void FMeshEditorMode::PlayFinishActionSound(EMeshEditAction::Type NewAction, UViewportInteractor* ActionInteractor)
+void FMeshEditorMode::PlayFinishActionSound( FName NewAction, UViewportInteractor* ActionInteractor)
 {
 	if (ActionInteractor != nullptr)
 	{
@@ -521,10 +523,8 @@ void FMeshEditorMode::BindCommands()
 
 	// Register editing modes (equipped actions)
 	RegisterVertexEditingMode( MeshEditorVertexActions.MoveVertex, EMeshEditAction::Move );
-	RegisterVertexEditingMode( MeshEditorVertexActions.EditVertexCornerSharpness, EMeshEditAction::EditVertexCornerSharpness );
 
 	RegisterEdgeEditingMode( MeshEditorEdgeActions.MoveEdge, EMeshEditAction::Move );
-	RegisterEdgeEditingMode( MeshEditorEdgeActions.EditEdgeCreaseSharpness, EMeshEditAction::EditEdgeCreaseSharpness );
 
 	RegisterPolygonEditingMode( MeshEditorPolygonActions.MovePolygon, EMeshEditAction::Move );
 	RegisterCommonEditingMode( MeshEditorCommonActions.DrawVertices, EMeshEditAction::DrawVertices );
@@ -620,14 +620,14 @@ void FMeshEditorMode::BindCommands()
 }
 
 
-void FMeshEditorMode::RegisterCommonEditingMode( const TSharedPtr<FUICommandInfo>& Command, EMeshEditAction::Type EditingMode )
+void FMeshEditorMode::RegisterCommonEditingMode( const TSharedPtr<FUICommandInfo>& Command, FName EditingMode )
 {
 	RegisterVertexEditingMode( Command, EditingMode );
 	RegisterEdgeEditingMode( Command, EditingMode );
 	RegisterPolygonEditingMode( Command, EditingMode );
 }
 
-void FMeshEditorMode::RegisterVertexEditingMode( const TSharedPtr<FUICommandInfo>& Command, EMeshEditAction::Type EditingMode )
+void FMeshEditorMode::RegisterVertexEditingMode( const TSharedPtr<FUICommandInfo>& Command, FName EditingMode )
 {
 	VertexActions.Emplace( Command, FUIAction(
 		FExecuteAction::CreateLambda( [this, EditingMode] { SetEquippedAction( EEditableMeshElementType::Vertex, EditingMode ); } ),
@@ -637,7 +637,7 @@ void FMeshEditorMode::RegisterVertexEditingMode( const TSharedPtr<FUICommandInfo
 }
 
 
-void FMeshEditorMode::RegisterEdgeEditingMode( const TSharedPtr<FUICommandInfo>& Command, EMeshEditAction::Type EditingMode )
+void FMeshEditorMode::RegisterEdgeEditingMode( const TSharedPtr<FUICommandInfo>& Command, FName EditingMode )
 {
 	EdgeActions.Emplace( Command, FUIAction(
 		FExecuteAction::CreateLambda( [this, EditingMode] { SetEquippedAction( EEditableMeshElementType::Edge, EditingMode ); } ),
@@ -647,7 +647,7 @@ void FMeshEditorMode::RegisterEdgeEditingMode( const TSharedPtr<FUICommandInfo>&
 }
 
 
-void FMeshEditorMode::RegisterPolygonEditingMode( const TSharedPtr<FUICommandInfo>& Command, EMeshEditAction::Type EditingMode )
+void FMeshEditorMode::RegisterPolygonEditingMode( const TSharedPtr<FUICommandInfo>& Command, FName EditingMode )
 {
 	PolygonActions.Emplace( Command, FUIAction(
 		FExecuteAction::CreateLambda( [this, EditingMode] { SetEquippedAction(EEditableMeshElementType::Polygon, EditingMode ); } ),
@@ -933,7 +933,7 @@ void FMeshEditorMode::Tick( FEditorViewportClient* ViewportClient, float DeltaTi
 		}
 	}
 
-	if( ActiveAction != EMeshEditAction::None &&
+	if( ActiveAction != NAME_None &&
 		ActiveAction != EMeshEditAction::SelectByPainting )		// When selecting, no updates are needed
 	{
 		const bool bIsActionFinishing = false;
@@ -965,7 +965,7 @@ void FMeshEditorMode::Tick( FEditorViewportClient* ViewportClient, float DeltaTi
 		const EEditableMeshElementType SelectedMeshElementType = GetSelectedMeshElementType();
 		ViewportWorldInteraction->SetTransformGizmoVisible(
 			( ActiveAction == EMeshEditAction::MoveUsingGizmo ) ||
-			( ActiveAction == EMeshEditAction::None &&
+			( ActiveAction == NAME_None &&
 				( ( EquippedPolygonAction == EMeshEditAction::Move && SelectedMeshElementType == EEditableMeshElementType::Polygon ) ||
 				  ( EquippedVertexAction == EMeshEditAction::Move && SelectedMeshElementType == EEditableMeshElementType::Vertex ) ||
 				  ( EquippedEdgeAction == EMeshEditAction::Move && SelectedMeshElementType == EEditableMeshElementType::Edge ) ) ) 
@@ -1230,7 +1230,7 @@ void FMeshEditorMode::DeselectMeshElements( const TMap<UEditableMesh*, TArray<FM
 
 bool FMeshEditorMode::DeleteSelectedMeshElement()
 {
-	if( ActiveAction != EMeshEditAction::None )
+	if( ActiveAction != NAME_None )
 	{
 		return false;
 	}
@@ -1310,7 +1310,7 @@ bool FMeshEditorMode::DeleteSelectedMeshElement()
 
 void FMeshEditorMode::AddOrRemoveSubdivisionLevel( const bool bShouldAdd )
 {
-	if( ActiveAction == EMeshEditAction::None )
+	if( ActiveAction == NAME_None )
 	{
 		static TSet< UEditableMesh* > SelectedMeshes;
 		SelectedMeshes.Reset();
@@ -1363,7 +1363,7 @@ void FMeshEditorMode::AddOrRemoveSubdivisionLevel( const bool bShouldAdd )
 
 void FMeshEditorMode::QuadrangulateMesh()
 {
-	if( ActiveAction == EMeshEditAction::None )
+	if( ActiveAction == NAME_None )
 	{
 		static TSet< UEditableMesh* > SelectedMeshes;
 		SelectedMeshes.Reset();
@@ -1502,7 +1502,7 @@ void FMeshEditorMode::FrameSelectedElements( FEditorViewportClient* ViewportClie
 
 bool FMeshEditorMode::RemoveSelectedEdges()
 {
-	if( ActiveAction != EMeshEditAction::None )
+	if( ActiveAction != NAME_None )
 	{
 		return false;
 	}
@@ -1591,7 +1591,7 @@ bool FMeshEditorMode::RemoveSelectedEdges()
 
 bool FMeshEditorMode::SelectEdgeLoops()
 {
-	if( ActiveAction != EMeshEditAction::None )
+	if( ActiveAction != NAME_None )
 	{
 		return false;
 	}
@@ -1643,7 +1643,7 @@ bool FMeshEditorMode::SelectEdgeLoops()
 
 bool FMeshEditorMode::RemoveSelectedVertices()
 {
-	if( ActiveAction != EMeshEditAction::None )
+	if( ActiveAction != NAME_None )
 	{
 		return false;
 	}
@@ -1735,7 +1735,7 @@ bool FMeshEditorMode::RemoveSelectedVertices()
 
 bool FMeshEditorMode::WeldSelectedVertices()
 {
-	if( ActiveAction != EMeshEditAction::None )
+	if( ActiveAction != NAME_None )
 	{
 		return false;
 	}
@@ -1821,7 +1821,7 @@ bool FMeshEditorMode::WeldSelectedVertices()
 
 bool FMeshEditorMode::FlipSelectedPolygons()
 {
-	if( ActiveAction != EMeshEditAction::None )
+	if( ActiveAction != NAME_None )
 	{
 		return false;
 	}
@@ -1871,7 +1871,7 @@ bool FMeshEditorMode::FlipSelectedPolygons()
 
 bool FMeshEditorMode::TriangulateSelectedPolygons()
 {
-	if( ActiveAction != EMeshEditAction::None )
+	if( ActiveAction != NAME_None )
 	{
 		return false;
 	}
@@ -1975,7 +1975,7 @@ bool FMeshEditorMode::AssignMaterialToSelectedPolygons( UMaterialInterface* Sele
 {
 	if( SelectedMaterial )
 	{
-		if( ActiveAction != EMeshEditAction::None )
+		if( ActiveAction != NAME_None )
 		{
 			return false;
 		}
@@ -2067,7 +2067,7 @@ bool FMeshEditorMode::AssignMaterialToSelectedPolygons( UMaterialInterface* Sele
 
 bool FMeshEditorMode::MakeSelectedEdgesHardOrSoft(const bool bMakeEdgesHard)
 {
-	if( ActiveAction != EMeshEditAction::None )
+	if( ActiveAction != NAME_None )
 	{
 		return false;
 	}
@@ -2510,7 +2510,7 @@ void FMeshEditorMode::Render( const FSceneView* SceneView, FViewport* Viewport,
 		}
 
 		// Only draw hover if we're not in the middle of an interactive edit
-		if( ActiveAction == EMeshEditAction::None )
+		if( ActiveAction == NAME_None )
 		{
 			for( FMeshEditorInteractorData& MeshEditorInteractorData : MeshEditorInteractorDatas )
 			{
@@ -2650,7 +2650,7 @@ void FMeshEditorMode::Render( const FSceneView* SceneView, FViewport* Viewport,
 	// Draw hovered and selected elements
 	{
 		// Only draw hover if we're not in the middle of an interactive edit
-		if( ActiveAction == EMeshEditAction::None )
+		if( ActiveAction == NAME_None )
 		{
 			const float HoveredSizeBias = MeshEd::HoveredSizeBias->GetFloat() + MeshEd::HoveredAnimationExtraSizeBias->GetFloat() * FMath::MakePulsatingValue( HoverFeedbackTimeValue, 0.5f );
 			{
@@ -2804,14 +2804,14 @@ void FMeshEditorMode::OnViewportInteractionHoverUpdate( UViewportInteractor* Vie
 		FViewportActionKeyInput* SelectAndMoveAction = ViewportInteractor->GetActionWithName( ViewportWorldActionTypes::SelectAndMove );
 		FViewportActionKeyInput* WorldMovementAction = ViewportInteractor->GetActionWithName( ViewportWorldActionTypes::WorldMovement );
  		const bool bIsLaserPointerBusy = 
-			( SelectAndMoveAction != nullptr && SelectAndMoveAction->bIsInputCaptured && ActiveAction == EMeshEditAction::None ) ||
-			( WorldMovementAction != nullptr && WorldMovementAction->bIsInputCaptured && ActiveAction == EMeshEditAction::None );
+			( SelectAndMoveAction != nullptr && SelectAndMoveAction->bIsInputCaptured && ActiveAction == NAME_None ) ||
+			( WorldMovementAction != nullptr && WorldMovementAction->bIsInputCaptured && ActiveAction == NAME_None );
 
 		bool bIsGrabberSphereOverMeshElement = false;
 
 		if( !bIsLaserPointerBusy )
 		{
-			if( ActiveAction == EMeshEditAction::None || 
+			if( ActiveAction == NAME_None || 
 				bActiveActionNeedsHoverLocation )
 			{
 				const float WorldSpaceRayFuzzyDistance = MeshEd::LaserFuzzySelectionDistance->GetFloat() * ViewportWorldInteraction->GetWorldScaleFactor();
@@ -3079,7 +3079,7 @@ void FMeshEditorMode::OnViewportInteractionInputUnhandled( FEditorViewportClient
 
 void FMeshEditorMode::OnViewportInteractionStartDragging( UViewportInteractor* ViewportInteractor )
 {
-	if( ActiveAction == EMeshEditAction::None )
+	if( ActiveAction == NAME_None )
 	{
 		// NOTE: We pass an empty Undo text string to tell StartAction() that we don't need it to start a transaction
 		// because the caller of this delegate will have already done that (the viewport interaction system)
@@ -3100,7 +3100,7 @@ void FMeshEditorMode::OnViewportInteractionStopDragging( UViewportInteractor* Vi
 
 void FMeshEditorMode::OnViewportInteractionFinishedMovingTransformables()
 {
-	if( ActiveAction != EMeshEditAction::None )
+	if( ActiveAction != NAME_None )
 	{
 		FinishAction();
 	}
@@ -3137,7 +3137,7 @@ void FMeshEditorMode::UpdateActiveAction( const bool bIsActionFinishing )
 	RollbackPreviewChanges();
 
 	if( bIsFirstActiveActionUpdate &&
-	   ActiveAction != EMeshEditAction::None &&
+	   ActiveAction != NAME_None &&
 	   ActiveAction != EMeshEditAction::SelectByPainting )
 	{
 		CommitSelectedMeshes();
@@ -3194,126 +3194,6 @@ void FMeshEditorMode::UpdateActiveAction( const bool bIsActionFinishing )
 	{
 		bIsMovingSelectedMeshElements = true;
 	}
-	else if( ActiveAction == EMeshEditAction::EditVertexCornerSharpness ||
-			 ActiveAction == EMeshEditAction::EditEdgeCreaseSharpness )
-	{
-		static TMap< UEditableMesh*, TArray<FMeshElement> > MeshesWithSelectedVertices;
-		static TMap< UEditableMesh*, TArray<FMeshElement> > MeshesWithSelectedEdges;
-		GetSelectedMeshesAndVertices( /* Out */ MeshesWithSelectedVertices );
-		GetSelectedMeshesAndEdges( /* Out */ MeshesWithSelectedEdges );
-
-		if( MeshesWithSelectedVertices.Num() > 0 || MeshesWithSelectedEdges.Num() > 0 )
-		{
-			// @todo mesheditor subdiv: Hard coded tweakables; ideally should be sized in screen space
-			const float DragScaleFactor = 5.0f;
-			const float ProgressBarHeight = 1000.0f;
-
-			// Figure out how much to either increase or decrease sharpness based on how far the user has dragged up or down in world space.
-			float DragDeltaZ = 0.0f;
-			
-			{
-				FMeshEditorInteractorData& MeshEditorInteractorData = GetMeshEditorInteractorData( ActiveActionInteractor );
-				if( MeshEditorInteractorData.bLaserIsValid || MeshEditorInteractorData.bGrabberSphereIsValid )
-				{
-					const FVector ProgressBarStart( EditSharpnessStartLocation + FVector( 0.0f, 0.0f, -ProgressBarHeight * 0.5f ) );
-					const FVector ProgressBarEnd( EditSharpnessStartLocation + FVector( 0.0f, 0.0f, ProgressBarHeight * 0.5f ) );
-
-					// First check to see if we're within grabber sphere range.  Otherwise, we'll use the laser.
-					bool bIsInRange = false;
-					FVector ClosestPointOnProgressBar;
-					if( MeshEditorInteractorData.bGrabberSphereIsValid )
-					{
-						if( FMath::PointDistToSegment( MeshEditorInteractorData.GrabberSphere.Center, ProgressBarStart, ProgressBarEnd ) <= MeshEditorInteractorData.GrabberSphere.W )
-						{
-							bIsInRange = true;
-							ClosestPointOnProgressBar = FMath::ClosestPointOnSegment(
-								MeshEditorInteractorData.GrabberSphere.Center,
-								ProgressBarStart, ProgressBarEnd );
-						}
-					}
-					
-					if( !bIsInRange && ensure( MeshEditorInteractorData.bLaserIsValid ) )
-					{
-						bIsInRange = true;
-						FVector ClosestPointOnRay;
-						FMath::SegmentDistToSegment(
-							ProgressBarStart, ProgressBarEnd,
-							MeshEditorInteractorData.LaserStart, MeshEditorInteractorData.LaserEnd,
-							/* Out */ ClosestPointOnProgressBar,
-							/* Out */ ClosestPointOnRay );
-					}
-
-					if( bIsInRange )
-					{
-						// Generate a drag value between -1.0 and 1.0 based on the interaction of the ray and the progress bar line segment.
-						const float ProgressBarLength = ( ProgressBarEnd - ProgressBarStart ).Size();
-						DragDeltaZ = -1.0f + 2.0f * ( ( ClosestPointOnProgressBar - ProgressBarStart ).Size() / ProgressBarLength );
-					}
-				}
-			}
-
-			const float ScaledDragDelta = DragDeltaZ * DragScaleFactor;
-
-			for( auto& MeshAndSelectedVertices : MeshesWithSelectedVertices )
-			{
-				UEditableMesh* EditableMesh = MeshAndSelectedVertices.Key;
-				const TArray<FMeshElement>& VertexElements = MeshAndSelectedVertices.Value;
-
-				static TArray<FVertexID> VertexIDs;
-				VertexIDs.Reset();
-
-				static TArray<float> NewSharpnessValues;
-				NewSharpnessValues.Reset();
-
-				for( const FMeshElement& VertexElement : VertexElements )
-				{
-					const FVertexID VertexID( VertexElement.ElementAddress.ElementID );
-
-					const float CurrentSharpnessValue = EditableMesh->GetVertexAttribute( VertexID, UEditableMeshAttribute::VertexCornerSharpness(), 0 ).X;
-					const float NewSharpnessValue = FMath::Clamp( CurrentSharpnessValue + ScaledDragDelta, 0.0f, 1.0f );
-
-					VertexIDs.Add( VertexID );
-					NewSharpnessValues.Add( NewSharpnessValue );
-				}
-
-
-				verify( !EditableMesh->AnyChangesToUndo() );
-
-				EditableMesh->SetVerticesCornerSharpness( VertexIDs, NewSharpnessValues );
-
-				TrackUndo( EditableMesh, EditableMesh->MakeUndo() );
-			}
-
-			for( auto& MeshAndSelectedEdges : MeshesWithSelectedEdges )
-			{
-				UEditableMesh* EditableMesh = MeshAndSelectedEdges.Key;
-				const TArray<FMeshElement>& EdgeElements = MeshAndSelectedEdges.Value;
-
-				static TArray<FEdgeID> EdgeIDs;
-				EdgeIDs.Reset();
-
-				static TArray<float> NewSharpnessValues;
-				NewSharpnessValues.Reset();
-
-				for( const FMeshElement& EdgeElement : EdgeElements )
-				{
-					const FEdgeID EdgeID( EdgeElement.ElementAddress.ElementID );
-
-					const float CurrentSharpnessValue = EditableMesh->GetEdgeAttribute( EdgeID, UEditableMeshAttribute::EdgeCreaseSharpness(), 0 ).X;
-					const float NewSharpnessValue = FMath::Clamp( CurrentSharpnessValue + ScaledDragDelta, 0.0f, 1.0f );
-
-					EdgeIDs.Add( EdgeID );
-					NewSharpnessValues.Add( NewSharpnessValue );
-				}
-
-				verify( !EditableMesh->AnyChangesToUndo() );
-
-				EditableMesh->SetEdgesCreaseSharpness( EdgeIDs, NewSharpnessValues );
-
-				TrackUndo( EditableMesh, EditableMesh->MakeUndo() );
-			}
-		}
-	}
 	else if( ActiveAction == EMeshEditAction::DrawVertices )
 	{
 		// @todo mesheditor: need a way to determine the plane we are going to create the polygon on.
@@ -4382,7 +4262,7 @@ void FMeshEditorMode::OnViewportInteractionInputAction( FEditorViewportClient& V
 
 		// If we're interactively editing something, clicking will commit that change
 		if( Action.Event == IE_Pressed &&
-			ActiveAction != EMeshEditAction::None )
+			ActiveAction != NAME_None )
 		{
 			// We're busy doing something else right now.  It might be an interactor trying to click while a different one is in the middle of something.
 			bWasHandled = true;
@@ -4391,7 +4271,7 @@ void FMeshEditorMode::OnViewportInteractionInputAction( FEditorViewportClient& V
 		// Otherwise, go ahead and try to interact with what's under the interactor
 		else if( Action.Event == IE_Pressed &&
 				 !bOutIsInputCaptured &&
-				 ActiveAction == EMeshEditAction::None )	// Only if we're not already doing something
+				 ActiveAction == NAME_None )	// Only if we're not already doing something
 		{
 			bool bWantToStartMoving = false;
 
@@ -4454,24 +4334,6 @@ void FMeshEditorMode::OnViewportInteractionInputAction( FEditorViewportClient& V
 						const bool bActionNeedsHoverLocation = false;
 						StartAction( EMeshEditAction::Move, ViewportInteractor, bActionNeedsHoverLocation, LOCTEXT( "UndoDragPolygon", "Drag Polygon" ) );
 					}
-					else if( SelectedMeshElementType == EEditableMeshElementType::Vertex && EquippedVertexAction == EMeshEditAction::EditVertexCornerSharpness )
-					{
-						EditSharpnessStartLocation = MeshEditorInteractorData.HoverLocation;
-
-						const bool bActionNeedsHoverLocation = false;
-						StartAction( EMeshEditAction::EditVertexCornerSharpness, ViewportInteractor, bActionNeedsHoverLocation, LOCTEXT( "UndoEditVertexCornerSharpness", "Edit Vertex Corner Sharpness" ) );
-
-						bOutIsInputCaptured = true;
-					}
-					else if( SelectedMeshElementType == EEditableMeshElementType::Edge && EquippedEdgeAction == EMeshEditAction::EditEdgeCreaseSharpness )
-					{
-						EditSharpnessStartLocation = MeshEditorInteractorData.HoverLocation;
-
-						const bool bActionNeedsHoverLocation = false;
-						StartAction( EMeshEditAction::EditEdgeCreaseSharpness, ViewportInteractor, bActionNeedsHoverLocation, LOCTEXT( "UndoEditEdgeCreaseSharpness", "Edit Edge Crease Sharpness" ) );
-
-						bOutIsInputCaptured = true;
-					}
 					else
 					{
 						for( TObjectIterator<UMeshEditorCommand> CommandCDOIter( RF_NoFlags ); CommandCDOIter; ++CommandCDOIter )
@@ -4479,7 +4341,7 @@ void FMeshEditorMode::OnViewportInteractionInputAction( FEditorViewportClient& V
 							UMeshEditorCommand* CommandCDO = *CommandCDOIter;
 							if( !( CommandCDO->GetClass()->GetClassFlags() & CLASS_Abstract ) )
 							{
-								EMeshEditAction::Type EquippedAction = EMeshEditAction::None;
+								FName EquippedAction = NAME_None;
 								switch( SelectedMeshElementType )
 								{
 									case EEditableMeshElementType::Vertex:
@@ -4593,7 +4455,7 @@ void FMeshEditorMode::OnViewportInteractionInputAction( FEditorViewportClient& V
 		}
 		else if( Action.Event == IE_Released )
 		{
-			if( ActiveAction != EMeshEditAction::None && 
+			if( ActiveAction != NAME_None && 
 				ActiveAction != EMeshEditAction::MoveUsingGizmo &&	// The button 'release' for gizmo-based movement is handled by the viewport world interaction system
 				bOutIsInputCaptured )
 			{
@@ -4628,10 +4490,10 @@ void FMeshEditorMode::OnViewportInteractionInputAction( FEditorViewportClient& V
 }
 
 
-void FMeshEditorMode::StartAction( EMeshEditAction::Type NewAction, UViewportInteractor* ActionInteractor, const bool bActionNeedsHoverLocation, const FText& UndoText )
+void FMeshEditorMode::StartAction( FName NewAction, UViewportInteractor* ActionInteractor, const bool bActionNeedsHoverLocation, const FText& UndoText )
 {
 	// Don't start a new action without finishing the previous one!
-	check( ActiveAction == EMeshEditAction::None );
+	check( ActiveAction == NAME_None );
 
 	PlayStartActionSound( NewAction, ActionInteractor );
 
@@ -4656,7 +4518,7 @@ void FMeshEditorMode::FinishAction()
 {
 	// @todo mesheditor: Make sure this is called before Undo is invoked (PreEditUndo!), otherwise the previous action will be undone instead of the active one
 
-	check( ActiveAction != EMeshEditAction::None );
+	check( ActiveAction != NAME_None );
 	check( GUndo == nullptr || GEditor->IsTransactionActive() );	// Someone must have started a transaction! (It might not have been us though.)
 
 	const bool bIsActionFinishing = true;
@@ -4680,7 +4542,7 @@ void FMeshEditorMode::FinishAction()
 		PlayFinishActionSound( ActiveAction, ActiveActionInteractor );
 	}
 
-	ActiveAction = EMeshEditAction::None;
+	ActiveAction = NAME_None;
 	ActiveActionInteractor = nullptr;
 	bActiveActionNeedsHoverLocation = false;
 
@@ -5046,7 +4908,7 @@ void FMeshEditorMode::RefreshTransformables()
 {
 	// Don't refresh transformables while we're actively moving them around
 	const bool bAllowRefresh =
-		ActiveAction == EMeshEditAction::None ||
+		ActiveAction == NAME_None ||
 		ActiveAction == EMeshEditAction::SelectByPainting ||
 		bIsFirstActiveActionUpdate;
 	if( !bAllowRefresh )
@@ -5352,9 +5214,9 @@ void FMeshEditorMode::MakeVRRadialMenuActionsMenu(FMenuBuilder& MenuBuilder, TSh
 }
 
 
-EMeshEditAction::Type FMeshEditorMode::GetEquippedAction( const EEditableMeshElementType ForElementType ) const
+FName FMeshEditorMode::GetEquippedAction( const EEditableMeshElementType ForElementType ) const
 {
-	EMeshEditAction::Type EquippedAction = EMeshEditAction::None;
+	FName EquippedAction = NAME_None;
 
 	switch( ForElementType )
 	{
@@ -5378,7 +5240,7 @@ EMeshEditAction::Type FMeshEditorMode::GetEquippedAction( const EEditableMeshEle
 }
 
 
-void FMeshEditorMode::SetEquippedAction( const EEditableMeshElementType ForElementType, const EMeshEditAction::Type ActionToEquip )
+void FMeshEditorMode::SetEquippedAction( const EEditableMeshElementType ForElementType, const FName ActionToEquip )
 {
 	switch( ForElementType )
 	{
diff --git a/Engine/Plugins/Editor/MeshEditor/Source/MeshEditor/MeshEditorMode.h b/Engine/Plugins/Editor/MeshEditor/Source/MeshEditor/MeshEditorMode.h
--- a/Engine/Plugins/Editor/MeshEditor/Source/MeshEditor/MeshEditorMode.h
+++ b/Engine/Plugins/Editor/MeshEditor/Source/MeshEditor/MeshEditorMode.h
@@ -163,12 +163,12 @@ protected:
 	virtual void SetEditingPerInstance( bool bPerInstance ) override { bPerInstanceEdits = bPerInstance; }
 	virtual void PropagateInstanceChanges() override;
 	virtual bool CanPropagateInstanceChanges() const override;
-	virtual EMeshEditAction::Type GetEquippedAction( const EEditableMeshElementType ForElementType ) const override;
-	virtual void SetEquippedAction( const EEditableMeshElementType ForElementType, const EMeshEditAction::Type ActionToEquip ) override;
+	virtual FName GetEquippedAction( const EEditableMeshElementType ForElementType ) const override;
+	virtual void SetEquippedAction( const EEditableMeshElementType ForElementType, const FName ActionToEquip ) override;
 
 	// IMeshEditorModeEditingContract interface
 	virtual const UEditableMesh* FindEditableMesh( class UPrimitiveComponent& Component, const FEditableMeshSubMeshAddress& SubMeshAddress ) const override;
-	virtual EMeshEditAction::Type GetActiveAction() const override { return ActiveAction; }
+	virtual FName GetActiveAction() const override { return ActiveAction; }
 	virtual void TrackUndo( UObject* Object, TUniquePtr<FChange> RevertChange ) override;
 	virtual void CommitSelectedMeshes() override;
 	virtual FMeshElement GetHoveredMeshElement( class UViewportInteractor* ViewportInteractor ) const override;
@@ -259,7 +259,7 @@ protected:
 	void UpdateCameraToWorldTransform( const FEditorViewportClient& ViewportClient );
 
 	/** Begins an action */
-	void StartAction( const EMeshEditAction::Type NewAction, class UViewportInteractor* ActionInteractor, const bool bActionNeedsHoverLocation, const FText& UndoText );
+	void StartAction( const FName NewAction, class UViewportInteractor* ActionInteractor, const bool bActionNeedsHoverLocation, const FText& UndoText );
 
 	/** Ends an action that's currently in progress.  Usually called when the user commits a change by clicking/releasing, but can
 	    also be called when the user begins a new action while inertia is still influencing the active action */
@@ -268,10 +268,10 @@ protected:
 	/** Binds UI commands to actions for the mesh editor */
 	void BindCommands();
 
-	void RegisterCommonEditingMode( const TSharedPtr<FUICommandInfo>& Command, EMeshEditAction::Type EditingMode );
-	void RegisterVertexEditingMode( const TSharedPtr<FUICommandInfo>& Command, EMeshEditAction::Type EditingMode );
-	void RegisterEdgeEditingMode( const TSharedPtr<FUICommandInfo>& Command, EMeshEditAction::Type EditingMode );
-	void RegisterPolygonEditingMode( const TSharedPtr<FUICommandInfo>& Command, EMeshEditAction::Type EditingMode );
+	void RegisterCommonEditingMode( const TSharedPtr<FUICommandInfo>& Command, FName EditingMode );
+	void RegisterVertexEditingMode( const TSharedPtr<FUICommandInfo>& Command, FName EditingMode );
+	void RegisterEdgeEditingMode( const TSharedPtr<FUICommandInfo>& Command, FName EditingMode );
+	void RegisterPolygonEditingMode( const TSharedPtr<FUICommandInfo>& Command, FName EditingMode );
 
 	void RegisterCommand( const TSharedPtr<FUICommandInfo>& Command, const FExecuteAction& ExecuteAction );
 	void RegisterVertexCommand( const TSharedPtr<FUICommandInfo>& Command, const FExecuteAction& ExecuteAction );
@@ -360,10 +360,10 @@ protected:
 	bool IsActive() const { return ( ViewportWorldInteraction != nullptr ); }
 
 	/** Plays sound when starting a mesh edit action */
-	void PlayStartActionSound(EMeshEditAction::Type NewAction, UViewportInteractor* ActionInteractor = nullptr);
+	void PlayStartActionSound( FName NewAction, UViewportInteractor* ActionInteractor = nullptr);
 
 	/** Plays sound when mesh edit action was finished */
-	void PlayFinishActionSound(EMeshEditAction::Type NewAction, UViewportInteractor* ActionInteractor = nullptr);
+	void PlayFinishActionSound( FName NewAction, UViewportInteractor* ActionInteractor = nullptr);
 
 
 protected:
@@ -530,17 +530,17 @@ protected:
 	FTrackingTransaction TrackingTransaction;
 
 	/** The next action that will be started when interacting with a selected vertex */
-	EMeshEditAction::Type EquippedVertexAction;
+	FName EquippedVertexAction;
 
 	/** The next action that will be started when interacting with a selected edge */
-	EMeshEditAction::Type EquippedEdgeAction;
+	FName EquippedEdgeAction;
 
 	/** The next action that will be started when interacting with a selected polygon */
-	EMeshEditAction::Type EquippedPolygonAction;
+	FName EquippedPolygonAction;
 
 	/** The interactive action currently being performed (and previewed).  These usually happen over multiple frames, and
 	    result in a 'final' application of the change that performs a more exhaustive (and more expensive) update. */
-	EMeshEditAction::Type ActiveAction;
+	FName ActiveAction;
 
 	/** Whether we're actually in the middile of updating the active action.  This means that StoreUndo() will behave
 	    differently in this case -- instead of pushing undo data to the editor, we'll capture it temporarily in PreviewRevertChanges,
@@ -586,14 +586,6 @@ protected:
 	TArray<TTuple<TSharedPtr<FUICommandInfo>, FUIAction>> PolygonActions;
 
 
-	//
-	// EditVertexCornerSharpness/EdgeEdgeCreaseSharpness
-	//
-
-	/** Where the active interactor's impact point was when the "edit sharpness" action started */
-	FVector EditSharpnessStartLocation;
-
-	
 	//
 	// DrawVertices
 	//
diff --git a/Engine/Plugins/Editor/MeshEditor/Source/MeshEditor/Public/IMeshEditorModeEditingContract.h b/Engine/Plugins/Editor/MeshEditor/Source/MeshEditor/Public/IMeshEditorModeEditingContract.h
--- a/Engine/Plugins/Editor/MeshEditor/Source/MeshEditor/Public/IMeshEditorModeEditingContract.h
+++ b/Engine/Plugins/Editor/MeshEditor/Source/MeshEditor/Public/IMeshEditorModeEditingContract.h
@@ -6,35 +6,6 @@
 #include "MeshElement.h"
 
 
-// @todo mesheditor extensibility: This should probably be removed after we've evicted all current mesh editing actions to another module
-namespace EMeshEditAction
-{
-	typedef FName Type;
-
-	/** Nothing going on right now */
-	MESHEDITOR_API extern const FName None;
-
-	/** Selecting mesh elements by 'painting' over multiple elements */
-	MESHEDITOR_API extern const FName SelectByPainting;
-
-	/** Moving elements using a transform gizmo */
-	MESHEDITOR_API extern const FName MoveUsingGizmo;
-
-	/** Moving selected mesh elements (vertices, edges or polygons) */
-	MESHEDITOR_API extern const FName Move;
-
-	/** For subdivision meshes, edits how sharp a vertex corner is by dragging in space */
-	MESHEDITOR_API extern const FName EditVertexCornerSharpness;
-
-	/** For subdivision meshes, edits how sharp an edge crease is by dragging in space */
-	MESHEDITOR_API extern const FName EditEdgeCreaseSharpness;
-
-	/** Freehand vertex drawing */
-	MESHEDITOR_API extern const FName DrawVertices;
-}
-
-
-
 class IMeshEditorModeEditingContract
 {
 
@@ -45,7 +16,7 @@ public:
 
 	/** Gets the interactive action currently being performed (and previewed).  These usually happen over multiple frames, and
 	    result in a 'final' application of the change that performs a more exhaustive (and more expensive) update. */
-	virtual EMeshEditAction::Type GetActiveAction() const = 0;
+	virtual FName GetActiveAction() const = 0;
 
 	/** Stores undo state for the specified object.  This will store the state different depending on whether we're
 	    currently in the middle of previewing a temporary change to meshes (bIsCapturingUndoForPreview) */
diff --git a/Engine/Plugins/Editor/MeshEditor/Source/MeshEditor/Public/IMeshEditorModeUIContract.h b/Engine/Plugins/Editor/MeshEditor/Source/MeshEditor/Public/IMeshEditorModeUIContract.h
--- a/Engine/Plugins/Editor/MeshEditor/Source/MeshEditor/Public/IMeshEditorModeUIContract.h
+++ b/Engine/Plugins/Editor/MeshEditor/Source/MeshEditor/Public/IMeshEditorModeUIContract.h
@@ -39,9 +39,9 @@ public:
 	virtual bool CanPropagateInstanceChanges() const = 0;
 
 	/** Returns the current action to use when interacting the next time with the specified type of mesh element selection mode */
-	virtual EMeshEditAction::Type GetEquippedAction( const EEditableMeshElementType ForElementType ) const = 0;
+	virtual FName GetEquippedAction( const EEditableMeshElementType ForElementType ) const = 0;
 
 	/** Sets the current action to use when interacting the next time with the specified type of mesh element selection mode */
-	virtual void SetEquippedAction( const EEditableMeshElementType ForElementType, const EMeshEditAction::Type ActionToEquip ) = 0;
+	virtual void SetEquippedAction( const EEditableMeshElementType ForElementType, const FName ActionToEquip ) = 0;
 
 };
diff --git a/Engine/Plugins/Editor/MeshEditor/Source/MeshEditor/Public/MeshEditorCommands.h b/Engine/Plugins/Editor/MeshEditor/Source/MeshEditor/Public/MeshEditorCommands.h
--- a/Engine/Plugins/Editor/MeshEditor/Source/MeshEditor/Public/MeshEditorCommands.h
+++ b/Engine/Plugins/Editor/MeshEditor/Source/MeshEditor/Public/MeshEditorCommands.h
@@ -217,9 +217,6 @@ public:
 	/** Sets the primary action to move vertices */
 	TSharedPtr<FUICommandInfo> MoveVertex;
 
-	/** Sets the primary action to edit the vertex's corner sharpness */
-	TSharedPtr<FUICommandInfo> EditVertexCornerSharpness;
-
 	/** Removes the selected vertex if possible */
 	TSharedPtr<FUICommandInfo> RemoveVertex;
 
@@ -243,9 +240,6 @@ public:
 	/** Sets the primary action to move edges */
 	TSharedPtr<FUICommandInfo> MoveEdge;
 
-	/** Sets the primary action to edit the edge's crease sharpness */
-	TSharedPtr<FUICommandInfo> EditEdgeCreaseSharpness;
-
 	/** Removes the selected edge if possible */
 	TSharedPtr<FUICommandInfo> RemoveEdge;
 
diff --git a/Engine/Plugins/Editor/MeshEditor/Source/PolygonModeling/EditVertexOrEdgeSharpness.cpp b/Engine/Plugins/Editor/MeshEditor/Source/PolygonModeling/EditVertexOrEdgeSharpness.cpp
new file mode 100644
--- /dev/null
+++ b/Engine/Plugins/Editor/MeshEditor/Source/PolygonModeling/EditVertexOrEdgeSharpness.cpp
@@ -0,0 +1,192 @@
+// Copyright 1998-2016 Epic Games, Inc. All Rights Reserved.
+
+#include "EditVertexOrEdgeSharpness.h"
+#include "IMeshEditorModeEditingContract.h"
+#include "IMeshEditorModeUIContract.h"
+#include "UICommandInfo.h"
+#include "EditableMesh.h"
+#include "MeshElement.h"
+#include "MultiBoxBuilder.h"
+#include "UICommandList.h"
+#include "ViewportInteractor.h"
+
+#define LOCTEXT_NAMESPACE "MeshEditorMode"
+
+
+namespace VertexOrEdgeSharpnessHelpers
+{
+	/** Where the active interactor's impact point was when the "edit sharpness" action started */
+	FVector EditSharpnessStartLocation( FVector::ZeroVector );
+
+	// @todo mesheditor extensibility: Get rid of all of the static stuff ideally and CDOs with state.  Have MeshEditorMode construct instances of commands.  Don't use TObjectIterator except at startup.
+
+	/** Figures out how much we should change the sharpness amount by looking at the interactor aim delta */
+	static float ComputeSharpnessChangeDelta( IMeshEditorModeEditingContract& MeshEditorMode, UViewportInteractor* ViewportInteractor )
+	{
+		// @todo mesheditor subdiv: Hard coded tweakables; ideally should be sized in screen space
+		const float DragScaleFactor = 5.0f;
+		const float ProgressBarHeight = 1000.0f;
+
+		// Figure out how much to either increase or decrease sharpness based on how far the user has dragged up or down in world space.
+		float DragDeltaZ = 0.0f;
+
+		FVector LaserStart, LaserEnd;
+		const bool bLaserIsValid = ViewportInteractor->GetLaserPointer( /* Out */ LaserStart, /* Out */ LaserEnd );
+		FSphere GrabberSphere;
+		const bool bGrabberSphereIsValid = ViewportInteractor->GetGrabberSphere( /* Out */ GrabberSphere );
+		if( bLaserIsValid || bGrabberSphereIsValid )
+		{
+			const FVector ProgressBarStart( EditSharpnessStartLocation + FVector( 0.0f, 0.0f, -ProgressBarHeight * 0.5f ) );
+			const FVector ProgressBarEnd( EditSharpnessStartLocation + FVector( 0.0f, 0.0f, ProgressBarHeight * 0.5f ) );
+
+			// First check to see if we're within grabber sphere range.  Otherwise, we'll use the laser.
+			bool bIsInRange = false;
+			FVector ClosestPointOnProgressBar;
+			if( bGrabberSphereIsValid )
+			{
+				if( FMath::PointDistToSegment( GrabberSphere.Center, ProgressBarStart, ProgressBarEnd ) <= GrabberSphere.W )
+				{
+					bIsInRange = true;
+					ClosestPointOnProgressBar = FMath::ClosestPointOnSegment(
+						GrabberSphere.Center,
+						ProgressBarStart, ProgressBarEnd );
+				}
+			}
+
+			if( !bIsInRange && ensure( bLaserIsValid ) )
+			{
+				bIsInRange = true;
+				FVector ClosestPointOnRay;
+				FMath::SegmentDistToSegment(
+					ProgressBarStart, ProgressBarEnd,
+					LaserStart, LaserEnd,
+					/* Out */ ClosestPointOnProgressBar,
+					/* Out */ ClosestPointOnRay );
+			}
+
+			if( bIsInRange )
+			{
+				// Generate a drag value between -1.0 and 1.0 based on the interaction of the ray and the progress bar line segment.
+				const float ProgressBarLength = ( ProgressBarEnd - ProgressBarStart ).Size();
+				DragDeltaZ = -1.0f + 2.0f * ( ( ClosestPointOnProgressBar - ProgressBarStart ).Size() / ProgressBarLength );
+			}
+		}
+
+		const float ScaledDragDelta = DragDeltaZ * DragScaleFactor;
+		return ScaledDragDelta;
+	}
+
+}
+
+
+void UEditVertexCornerSharpnessCommand::RegisterUICommand( FBindingContext* BindingContext )
+{
+	UI_COMMAND_EXT( BindingContext, /* Out */ UICommandInfo, "EditVertexCornerSharpness", "Edit Vertex Corner Sharpness Mode", "Set the primary action to edit vertex corner sharpness.", EUserInterfaceActionType::RadioButton, FInputChord() );
+}
+
+
+bool UEditVertexCornerSharpnessCommand::TryStartingToDrag( IMeshEditorModeEditingContract& MeshEditorMode, UViewportInteractor* ViewportInteractor )
+{
+	VertexOrEdgeSharpnessHelpers::EditSharpnessStartLocation = ViewportInteractor->GetHoverLocation();
+
+	return true;
+}
+
+
+void UEditVertexCornerSharpnessCommand::ApplyDuringDrag( IMeshEditorModeEditingContract& MeshEditorMode, UViewportInteractor* ViewportInteractor )
+{
+	static TMap< UEditableMesh*, TArray<FMeshElement> > MeshesWithSelectedVertices;
+	MeshEditorMode.GetSelectedMeshesAndVertices( /* Out */ MeshesWithSelectedVertices );
+	if( MeshesWithSelectedVertices.Num() > 0 )
+	{
+		const float ScaledDragDelta = VertexOrEdgeSharpnessHelpers::ComputeSharpnessChangeDelta( MeshEditorMode, ViewportInteractor );
+
+		for( auto& MeshAndSelectedVertices : MeshesWithSelectedVertices )
+		{
+			UEditableMesh* EditableMesh = MeshAndSelectedVertices.Key;
+			const TArray<FMeshElement>& VertexElements = MeshAndSelectedVertices.Value;
+
+			static TArray<FVertexID> VertexIDs;
+			VertexIDs.Reset();
+
+			static TArray<float> NewSharpnessValues;
+			NewSharpnessValues.Reset();
+
+			for( const FMeshElement& VertexElement : VertexElements )
+			{
+				const FVertexID VertexID( VertexElement.ElementAddress.ElementID );
+
+				const float CurrentSharpnessValue = EditableMesh->GetVertexAttribute( VertexID, UEditableMeshAttribute::VertexCornerSharpness(), 0 ).X;
+				const float NewSharpnessValue = FMath::Clamp( CurrentSharpnessValue + ScaledDragDelta, 0.0f, 1.0f );
+
+				VertexIDs.Add( VertexID );
+				NewSharpnessValues.Add( NewSharpnessValue );
+			}
+
+
+			verify( !EditableMesh->AnyChangesToUndo() );
+
+			EditableMesh->SetVerticesCornerSharpness( VertexIDs, NewSharpnessValues );
+
+			MeshEditorMode.TrackUndo( EditableMesh, EditableMesh->MakeUndo() );
+		}
+	}
+}
+
+
+void UEditEdgeCreaseSharpnessCommand::RegisterUICommand( FBindingContext* BindingContext )
+{
+	UI_COMMAND_EXT( BindingContext, /* Out */ UICommandInfo, "EditEdgeCreaseSharpness", "Edit Edge Crease Sharpness Mode", "Set the primary action to edit edge crease sharpness.", EUserInterfaceActionType::RadioButton, FInputChord() );
+}
+
+
+bool UEditEdgeCreaseSharpnessCommand::TryStartingToDrag( IMeshEditorModeEditingContract& MeshEditorMode, UViewportInteractor* ViewportInteractor )
+{
+	VertexOrEdgeSharpnessHelpers::EditSharpnessStartLocation = ViewportInteractor->GetHoverLocation();
+
+	return true;
+}
+
+
+void UEditEdgeCreaseSharpnessCommand::ApplyDuringDrag( IMeshEditorModeEditingContract& MeshEditorMode, UViewportInteractor* ViewportInteractor )
+{
+	static TMap< UEditableMesh*, TArray<FMeshElement> > MeshesWithSelectedEdges;
+	MeshEditorMode.GetSelectedMeshesAndEdges( /* Out */ MeshesWithSelectedEdges );
+	if( MeshesWithSelectedEdges.Num() > 0 )
+	{
+		const float ScaledDragDelta = VertexOrEdgeSharpnessHelpers::ComputeSharpnessChangeDelta( MeshEditorMode, ViewportInteractor );
+
+		for( auto& MeshAndSelectedEdges : MeshesWithSelectedEdges )
+		{
+			UEditableMesh* EditableMesh = MeshAndSelectedEdges.Key;
+			const TArray<FMeshElement>& EdgeElements = MeshAndSelectedEdges.Value;
+
+			static TArray<FEdgeID> EdgeIDs;
+			EdgeIDs.Reset();
+
+			static TArray<float> NewSharpnessValues;
+			NewSharpnessValues.Reset();
+
+			for( const FMeshElement& EdgeElement : EdgeElements )
+			{
+				const FEdgeID EdgeID( EdgeElement.ElementAddress.ElementID );
+
+				const float CurrentSharpnessValue = EditableMesh->GetEdgeAttribute( EdgeID, UEditableMeshAttribute::EdgeCreaseSharpness(), 0 ).X;
+				const float NewSharpnessValue = FMath::Clamp( CurrentSharpnessValue + ScaledDragDelta, 0.0f, 1.0f );
+
+				EdgeIDs.Add( EdgeID );
+				NewSharpnessValues.Add( NewSharpnessValue );
+			}
+
+			verify( !EditableMesh->AnyChangesToUndo() );
+
+			EditableMesh->SetEdgesCreaseSharpness( EdgeIDs, NewSharpnessValues );
+
+			MeshEditorMode.TrackUndo( EditableMesh, EditableMesh->MakeUndo() );
+		}
+	}
+}
+
+
+#undef LOCTEXT_NAMESPACE
+
diff --git a/Engine/Plugins/Editor/MeshEditor/Source/PolygonModeling/EditVertexOrEdgeSharpness.h b/Engine/Plugins/Editor/MeshEditor/Source/PolygonModeling/EditVertexOrEdgeSharpness.h
new file mode 100644
--- /dev/null
+++ b/Engine/Plugins/Editor/MeshEditor/Source/PolygonModeling/EditVertexOrEdgeSharpness.h
@@ -0,0 +1,52 @@
+// Copyright 1998-2016 Epic Games, Inc. All Rights Reserved.
+
+#pragma once
+
+#include "MeshEditorCommands.h"
+#include "EditVertexOrEdgeSharpness.generated.h"
+
+
+/** For subdivision meshes, edits how sharp a vertex corner is by dragging in space */
+UCLASS()
+class UEditVertexCornerSharpnessCommand : public UMeshEditorVertexCommand
+{
+	GENERATED_BODY()
+
+protected:
+
+	UEditVertexCornerSharpnessCommand()
+	{
+		bIsMode = true;
+		UndoText = NSLOCTEXT( "MeshEditor", "UndoEditVertexCornerSharpness", "Edit Vertex Corner Sharpness" );
+		bNeedsHoverLocation = false;
+		bNeedsDraggingInitiated = false;
+	}
+
+	// Overrides
+	virtual void RegisterUICommand( class FBindingContext* BindingContext ) override;
+	virtual bool TryStartingToDrag( IMeshEditorModeEditingContract& MeshEditorMode, class UViewportInteractor* ViewportInteractor ) override;
+	virtual void ApplyDuringDrag( IMeshEditorModeEditingContract& MeshEditorMode, class UViewportInteractor* ViewportInteractor ) override;
+};
+
+
+/** For subdivision meshes, edits how sharp an edge crease is by dragging in space */
+UCLASS()
+class UEditEdgeCreaseSharpnessCommand : public UMeshEditorEdgeCommand
+{
+	GENERATED_BODY()
+
+protected:
+
+	UEditEdgeCreaseSharpnessCommand()
+	{
+		bIsMode = true;
+		UndoText = NSLOCTEXT( "MeshEditor", "UndoEditEdgeCreaseSharpness", "Edit Edge Crease Sharpness" );
+		bNeedsHoverLocation = false;
+		bNeedsDraggingInitiated = false;
+	}
+
+	// Overrides
+	virtual void RegisterUICommand( class FBindingContext* BindingContext ) override;
+	virtual bool TryStartingToDrag( IMeshEditorModeEditingContract& MeshEditorMode, class UViewportInteractor* ViewportInteractor ) override;
+	virtual void ApplyDuringDrag( IMeshEditorModeEditingContract& MeshEditorMode, class UViewportInteractor* ViewportInteractor ) override;
+};
diff --git a/Engine/Plugins/Editor/MeshEditor/Source/PolygonModeling/TessellatePolygon.cpp b/Engine/Plugins/Editor/MeshEditor/Source/PolygonModeling/TessellatePolygon.cpp
--- a/Engine/Plugins/Editor/MeshEditor/Source/PolygonModeling/TessellatePolygon.cpp
+++ b/Engine/Plugins/Editor/MeshEditor/Source/PolygonModeling/TessellatePolygon.cpp
@@ -18,7 +18,7 @@ void UTessellatePolygonCommand::RegisterUICommand( FBindingContext* BindingConte
 
 void UTessellatePolygonCommand::Execute( IMeshEditorModeEditingContract& MeshEditorMode )
 {
-	if( MeshEditorMode.GetActiveAction() != EMeshEditAction::None )
+	if( MeshEditorMode.GetActiveAction() != NAME_None )
 	{
 		return;
 	}
