Mesh Editor extensibility (phase 7)
- Migrated Remove Edge, Remove Vertex, Harden Edge and Soften Edge to a new module
- Fixed a bug where removing an edge would not reselect the original edge after undo

diff --git a/Engine/Plugins/Editor/MeshEditor/Source/MeshEditor/MeshEditorCommands.cpp b/Engine/Plugins/Editor/MeshEditor/Source/MeshEditor/MeshEditorCommands.cpp
--- a/Engine/Plugins/Editor/MeshEditor/Source/MeshEditor/MeshEditorCommands.cpp
+++ b/Engine/Plugins/Editor/MeshEditor/Source/MeshEditor/MeshEditorCommands.cpp
@@ -82,7 +82,6 @@ void FMeshEditorVertexCommands::RegisterCommands()
 {
 	UI_COMMAND(MoveVertex, "Move Vertex Mode", "Set the primary action to move vertices.", EUserInterfaceActionType::RadioButton, FInputChord(EKeys::F1));
 
-	UI_COMMAND(RemoveVertex, "Remove Vertex", "Remove the selected vertex if possible.", EUserInterfaceActionType::Button, FInputChord(EKeys::BackSpace));
 	UI_COMMAND(WeldVertices, "Weld Vertices", "Weld the selected vertices, keeping the first selected vertex.", EUserInterfaceActionType::Button, FInputChord());
 
 	// @todo mesheditor extensibility: What's the plan for default keybinds for commands registered in a modular way?  Should we suggest available keys?
@@ -110,9 +109,6 @@ void FMeshEditorEdgeCommands::RegisterCommands()
 {
 	UI_COMMAND(MoveEdge, "Move Edge Mode", "Set the primary action to move edges.", EUserInterfaceActionType::RadioButton, FInputChord(EKeys::F1));
 
-	UI_COMMAND(RemoveEdge, "Remove Edge", "Remove the selected edge if possible.", EUserInterfaceActionType::Button, FInputChord(EKeys::BackSpace));
-	UI_COMMAND(SoftenEdge, "Soften Edge", "Make selected edge soft.", EUserInterfaceActionType::Button, FInputChord(EKeys::H, EModifierKey::Shift));
-	UI_COMMAND(HardenEdge, "Harden Edge", "Make selected edge hard.", EUserInterfaceActionType::Button, FInputChord(EKeys::H));
 	UI_COMMAND(SelectEdgeLoop, "Select Edge Loop", "Select the edge loops which contain the selected edges.", EUserInterfaceActionType::Button, FInputChord(EKeys::Two, EModifierKey::Shift));
 
 	for( TObjectIterator<UMeshEditorEdgeCommand> EdgeCommandCDOIter( RF_NoFlags ); EdgeCommandCDOIter; ++EdgeCommandCDOIter )
diff --git a/Engine/Plugins/Editor/MeshEditor/Source/MeshEditor/MeshEditorMode.cpp b/Engine/Plugins/Editor/MeshEditor/Source/MeshEditor/MeshEditorMode.cpp
--- a/Engine/Plugins/Editor/MeshEditor/Source/MeshEditor/MeshEditorMode.cpp
+++ b/Engine/Plugins/Editor/MeshEditor/Source/MeshEditor/MeshEditorMode.cpp
@@ -551,12 +551,8 @@ void FMeshEditorMode::BindCommands()
 	RegisterCommand( MeshEditorCommonActions.QuadrangulateMesh, FExecuteAction::CreateLambda( [ this ] { QuadrangulateMesh(); } ) );
 
 	// Register element-specific commands
-	RegisterVertexCommand( MeshEditorVertexActions.RemoveVertex, FExecuteAction::CreateLambda( [this] { RemoveSelectedVertices(); } ) );
 	RegisterVertexCommand( MeshEditorVertexActions.WeldVertices, FExecuteAction::CreateLambda( [this] { WeldSelectedVertices(); } ) );
 
-	RegisterEdgeCommand( MeshEditorEdgeActions.RemoveEdge, FExecuteAction::CreateLambda( [this] { RemoveSelectedEdges(); } ) );
-	RegisterEdgeCommand( MeshEditorEdgeActions.SoftenEdge, FExecuteAction::CreateLambda( [this] { MakeSelectedEdgesHardOrSoft( false ); } ) );
-	RegisterEdgeCommand( MeshEditorEdgeActions.HardenEdge, FExecuteAction::CreateLambda( [this] { MakeSelectedEdgesHardOrSoft( true ); } ) );
 	RegisterEdgeCommand( MeshEditorEdgeActions.SelectEdgeLoop, FExecuteAction::CreateLambda( [ this ] { SelectEdgeLoops(); } ) );
 
 	RegisterPolygonCommand( MeshEditorPolygonActions.FlipPolygon, FExecuteAction::CreateLambda( [this] { FlipSelectedPolygons(); } ) );
@@ -1500,95 +1496,6 @@ void FMeshEditorMode::FrameSelectedElements( FEditorViewportClient* ViewportClie
 }
 
 
-bool FMeshEditorMode::RemoveSelectedEdges()
-{
-	if( ActiveAction != NAME_None )
-	{
-		return false;
-	}
-
-	// @todo mesheditor: Only if one edge is selected, for now.  It gets a bit confusing when performing this operation
-	// with more than one edge selected.  However, it can be very useful when collapsing away edges that don't share any
-	// common polygons, so we should try to support it.
-	if( SelectedMeshElements.Num() != 1 )
-	{
-		// @todo should this count as a failure case?
-		return false;
-	}
-
-	static TMap< UEditableMesh*, TArray< FMeshElement > > MeshesWithEdgesToRemove;
-	GetSelectedMeshesAndEdges( MeshesWithEdgesToRemove );
-
-	if( MeshesWithEdgesToRemove.Num() == 0 )
-	{
-		// @todo should this count as a failure case?
-		return false;
-	}
-
-	FScopedTransaction Transaction( LOCTEXT( "UndoRemoveEdge", "Remove Edge" ) );
-
-	CommitSelectedMeshes();
-
-	// Refresh selection (committing may have created a new mesh instance)
-	GetSelectedMeshesAndEdges( MeshesWithEdgesToRemove );
-
-	TArray<FMeshElement> MeshElementsToSelect;
-	for( const auto& MeshAndElements : MeshesWithEdgesToRemove )
-	{
-		UEditableMesh* EditableMesh = MeshAndElements.Key;
-
-		EditableMesh->StartModification( EMeshModificationType::Final, EMeshTopologyChange::TopologyChange );
-
-		const TArray<FMeshElement>& EdgeElementsToRemove = MeshAndElements.Value;
-
-		for( const FMeshElement& EdgeElementToRemove : EdgeElementsToRemove )
-		{
-			const FEdgeID EdgeID( EdgeElementToRemove.ElementAddress.ElementID );
-			{
-				bool bWasEdgeRemoved = false;
-				FPolygonRef NewPolygonRef;
-
-				EditableMesh->TryToRemovePolygonEdge( EdgeID, /* Out */ bWasEdgeRemoved, /* Out */ NewPolygonRef );
-
-				if( bWasEdgeRemoved )
-				{
-					// Select the new polygon
-					FMeshElement NewPolygonMeshElement;
-					{
-						NewPolygonMeshElement.Component = EdgeElementToRemove.Component;
-						NewPolygonMeshElement.ElementAddress = EdgeElementToRemove.ElementAddress;
-						NewPolygonMeshElement.ElementAddress.ElementType = EEditableMeshElementType::Polygon;
-						NewPolygonMeshElement.ElementAddress.SectionID = NewPolygonRef.SectionID;
-						NewPolygonMeshElement.ElementAddress.ElementID = NewPolygonRef.PolygonID;
-					}
-
-					MeshElementsToSelect.Add( NewPolygonMeshElement );
-				}
-				else
-				{
-					// Couldn't remove the edge
-					// @todo mesheditor: Needs good user feedback when this happens
-				}
-			}
-		}
-
-		EditableMesh->EndModification();
-
-		TrackUndo( EditableMesh, EditableMesh->MakeUndo() );
-	}
-
-	{
-		// Make sure we're not still hovering over a polygon we removed
-		ClearInvalidSelectedElements();
-
-		// Select the polygon leftover after removing the edge
-		SelectMeshElements( MeshElementsToSelect );
-	}
-
-	return true;
-}
-
-
 bool FMeshEditorMode::SelectEdgeLoops()
 {
 	if( ActiveAction != NAME_None )
@@ -1641,98 +1548,6 @@ bool FMeshEditorMode::SelectEdgeLoops()
 }
 
 
-bool FMeshEditorMode::RemoveSelectedVertices()
-{
-	if( ActiveAction != NAME_None )
-	{
-		return false;
-	}
-
-	// @todo mesheditor: Only if one vertex is selected, for now.  It gets a bit confusing when performing this operation
-	// with more than one vertex selected.  However, it can be very useful when collapsing away vertices that don't share any
-	// common polygons, so we should try to support it.
-	if( SelectedMeshElements.Num() != 1 )
-	{
-		// @todo should this count as a failure case?
-		return false;
-	}
-
-	static TMap< UEditableMesh*, TArray<FMeshElement> > MeshesWithVerticesToRemove;
-	GetSelectedMeshesAndVertices( MeshesWithVerticesToRemove );
-
-	// Remove selected vertices
-	if( MeshesWithVerticesToRemove.Num() == 0 )
-	{
-		// @todo should this count as a failure case?
-		return false;
-	}
-
-	FScopedTransaction Transaction( LOCTEXT( "UndoRemoveVertex", "Remove Vertex" ) );
-
-	CommitSelectedMeshes();
-
-	// Refresh selection (committing may have created a new mesh instance)
-	GetSelectedMeshesAndVertices( MeshesWithVerticesToRemove );
-
-	// Deselect the mesh elements before we delete them.  This will make sure they become selected again after undo.
-	DeselectMeshElements( MeshesWithVerticesToRemove );
-
-	TArray<FMeshElement> MeshElementsToSelect;
-	for( const auto& MeshAndElements : MeshesWithVerticesToRemove )
-	{
-		UEditableMesh* EditableMesh = MeshAndElements.Key;
-
-		EditableMesh->StartModification( EMeshModificationType::Final, EMeshTopologyChange::TopologyChange );
-
-		const TArray<FMeshElement>& VertexElementsToRemove = MeshAndElements.Value;
-
-		for( const FMeshElement& VertexElementToRemove : VertexElementsToRemove )
-		{
-			const FVertexID VertexID( VertexElementToRemove.ElementAddress.ElementID );
-			{
-				bool bWasVertexRemoved = false;
-				FEdgeID NewEdgeID = FEdgeID::Invalid;
-
-				EditableMesh->TryToRemoveVertex( VertexID, /* Out */ bWasVertexRemoved, /* Out */ NewEdgeID );
-
-				if( bWasVertexRemoved )
-				{
-					// Select the new edge
-					FMeshElement NewEdgeMeshElement;
-					{
-						NewEdgeMeshElement.Component = VertexElementToRemove.Component;
-						NewEdgeMeshElement.ElementAddress = VertexElementToRemove.ElementAddress;
-						NewEdgeMeshElement.ElementAddress.ElementType = EEditableMeshElementType::Edge;
-						NewEdgeMeshElement.ElementAddress.ElementID = NewEdgeID;
-					}
-
-					MeshElementsToSelect.Add( NewEdgeMeshElement );
-				}
-				else
-				{
-					// Couldn't remove the vertex
-					// @todo mesheditor: Needs good user feedback when this happens
-					// @todo mesheditor: If this fails, it will already potentially have created a new instance. To be 100% correct, it needs to do a prepass
-					// to determine whether the operation can complete successfully before actually doing it.
-				}
-			}
-		}
-
-		EditableMesh->EndModification();
-
-		TrackUndo( EditableMesh, EditableMesh->MakeUndo() );
-	}
-
-	{
-		// Make sure we're not still hovering over a polygon we removed
-		ClearInvalidSelectedElements();
-		SelectMeshElements( MeshElementsToSelect );
-	}
-
-	return true;
-}
-
-
 bool FMeshEditorMode::WeldSelectedVertices()
 {
 	if( ActiveAction != NAME_None )
@@ -2065,63 +1880,6 @@ bool FMeshEditorMode::AssignMaterialToSelectedPolygons( UMaterialInterface* Sele
 }
 
 
-bool FMeshEditorMode::MakeSelectedEdgesHardOrSoft(const bool bMakeEdgesHard)
-{
-	if( ActiveAction != NAME_None )
-	{
-		return false;
-	}
-
-	// Make edges hard or soft
-	static TMap< UEditableMesh*, TArray<FMeshElement> > MeshesAndEdges;
-	GetSelectedMeshesAndEdges( /* Out */ MeshesAndEdges );
-
-	if( MeshesAndEdges.Num() == 0 )
-	{
-		// @todo should this count as a failure case?
-		return false;
-	}
-
-	const FScopedTransaction Transaction( bMakeEdgesHard ? LOCTEXT( "UndoHardenEdge", "Harden Edge" ) : LOCTEXT( "UndoSoftenEdge", "Soften Edge" ) );
-
-	CommitSelectedMeshes();
-
-	// Refresh selection (committing may have created a new mesh instance)
-	GetSelectedMeshesAndEdges( /* Out */ MeshesAndEdges );
-
-	for( auto& MeshAndEdges : MeshesAndEdges )
-	{
-		static TArray<FEdgeID> EdgeIDs;
-		EdgeIDs.Reset();
-
-		static TArray<bool> EdgesNewIsHard;
-		EdgesNewIsHard.Reset();
-
-		for( FMeshElement& EdgeElement : MeshAndEdges.Value )
-		{
-			const FEdgeID EdgeID( EdgeElement.ElementAddress.ElementID );
-
-			EdgeIDs.Add( EdgeID );
-			EdgesNewIsHard.Add( bMakeEdgesHard );
-		}
-
-
-		UEditableMesh* EditableMesh = MeshAndEdges.Key;
-		verify( !EditableMesh->AnyChangesToUndo() );
-
-		EditableMesh->StartModification( EMeshModificationType::Final, EMeshTopologyChange::TopologyChange );
-
-		EditableMesh->SetEdgesHardness( EdgeIDs, EdgesNewIsHard );
-
-		EditableMesh->EndModification();
-
-		TrackUndo( EditableMesh, EditableMesh->MakeUndo() );
-	}
-
-	return true;
-}
-
-
 bool FMeshEditorMode::InputAxis( FEditorViewportClient* ViewportClient, FViewport* Viewport, int32 ControllerId, FKey Key, float Delta, float DeltaTime )
 {
 	bool bHandled = false;
@@ -5125,18 +4883,6 @@ void FMeshEditorMode::MakeVRRadialMenuActionsMenu(FMenuBuilder& MenuBuilder, TSh
 			NAME_None,
 			EUserInterfaceActionType::CollapsedButton
 			);
-		MenuBuilder.AddMenuEntry(
-			LOCTEXT("RemoveSelected", "Remove Selected"),
-			FText(),
-			FSlateIcon(FMeshEditorStyle::GetStyleSetName(), "MeshEditorMode.EdgeRemove"),
-			FUIAction
-			(
-				FExecuteAction::CreateLambda([this] { RemoveSelectedEdges(); }),
-				FCanExecuteAction::CreateSP(this, &FMeshEditorMode::IsMeshElementTypeSelected, EEditableMeshElementType::Edge)
-				),
-			NAME_None,
-			EUserInterfaceActionType::CollapsedButton
-			);
 		MenuBuilder.AddMenuEntry(
 			LOCTEXT( "SelectEdgeLoop", "Select Edge Loop" ),
 			FText(),
@@ -5177,18 +4923,6 @@ void FMeshEditorMode::MakeVRRadialMenuActionsMenu(FMenuBuilder& MenuBuilder, TSh
 			NAME_None,
 			EUserInterfaceActionType::CollapsedButton
 			);
-		MenuBuilder.AddMenuEntry(
-			LOCTEXT("RemoveSelected", "Remove Selected"),
-			FText(),
-			FSlateIcon(FMeshEditorStyle::GetStyleSetName(), "MeshEditorMode.VertexRemove"),
-			FUIAction
-			(
-				FExecuteAction::CreateLambda([this] { RemoveSelectedVertices(); }),
-				FCanExecuteAction::CreateSP(this, &FMeshEditorMode::IsMeshElementTypeSelected, EEditableMeshElementType::Vertex)
-			),
-			NAME_None,
-			EUserInterfaceActionType::CollapsedButton
-			);
 		MenuBuilder.AddMenuEntry(
 			LOCTEXT("WeldSelected", "Weld Selected"),
 			FText(),
diff --git a/Engine/Plugins/Editor/MeshEditor/Source/MeshEditor/MeshEditorMode.h b/Engine/Plugins/Editor/MeshEditor/Source/MeshEditor/MeshEditorMode.h
--- a/Engine/Plugins/Editor/MeshEditor/Source/MeshEditor/MeshEditorMode.h
+++ b/Engine/Plugins/Editor/MeshEditor/Source/MeshEditor/MeshEditorMode.h
@@ -298,15 +298,9 @@ protected:
 	/** Moves the viewport camera to frame the currently selected elements */
 	void FrameSelectedElements( FEditorViewportClient* ViewportClient );
 
-	/** Removes the selected edge if possible; returns whether successful */
-	bool RemoveSelectedEdges();
-
 	/** Selects the edge loops which contain the selected edges */
 	bool SelectEdgeLoops();
 
-	/** Removes the selected vertex if possible; returns whether successful */
-	bool RemoveSelectedVertices();
-
 	/** Welds the selected vertices if possible, keeping the first selected vertex */
 	bool WeldSelectedVertices();
 
@@ -322,9 +316,6 @@ protected:
 	/** Assigns a material to the selected polygons; returns whether successful */
 	bool AssignMaterialToSelectedPolygons( UMaterialInterface* SelectedMaterial );
 
-	/** Creases selected edges; returns whether successful */
-	bool MakeSelectedEdgesHardOrSoft( bool bMakeEdgesHard );
-
 	/** Rolls back whatever we changed last time while previewing */
 	void RollbackPreviewChanges();
 
diff --git a/Engine/Plugins/Editor/MeshEditor/Source/MeshEditor/Public/MeshEditorCommands.h b/Engine/Plugins/Editor/MeshEditor/Source/MeshEditor/Public/MeshEditorCommands.h
--- a/Engine/Plugins/Editor/MeshEditor/Source/MeshEditor/Public/MeshEditorCommands.h
+++ b/Engine/Plugins/Editor/MeshEditor/Source/MeshEditor/Public/MeshEditorCommands.h
@@ -217,9 +217,6 @@ public:
 	/** Sets the primary action to move vertices */
 	TSharedPtr<FUICommandInfo> MoveVertex;
 
-	/** Removes the selected vertex if possible */
-	TSharedPtr<FUICommandInfo> RemoveVertex;
-
 	/** Welds the selected vertices */
 	TSharedPtr<FUICommandInfo> WeldVertices;
 };
@@ -240,15 +237,6 @@ public:
 	/** Sets the primary action to move edges */
 	TSharedPtr<FUICommandInfo> MoveEdge;
 
-	/** Removes the selected edge if possible */
-	TSharedPtr<FUICommandInfo> RemoveEdge;
-
-	/** Soften edge */
-	TSharedPtr<FUICommandInfo> SoftenEdge;
-
-	/** Harden edge */
-	TSharedPtr<FUICommandInfo> HardenEdge;
-
 	/** Select edge loop */
 	TSharedPtr<FUICommandInfo> SelectEdgeLoop;
 };
diff --git a/Engine/Plugins/Editor/MeshEditor/Source/PolygonModeling/HardenOrSoftenEdge.cpp b/Engine/Plugins/Editor/MeshEditor/Source/PolygonModeling/HardenOrSoftenEdge.cpp
new file mode 100644
--- /dev/null
+++ b/Engine/Plugins/Editor/MeshEditor/Source/PolygonModeling/HardenOrSoftenEdge.cpp
@@ -0,0 +1,97 @@
+// Copyright 1998-2016 Epic Games, Inc. All Rights Reserved.
+
+#include "HardenOrSoftenEdge.h"
+#include "IMeshEditorModeEditingContract.h"
+#include "UICommandInfo.h"
+#include "EditableMesh.h"
+#include "MeshElement.h"
+
+
+#define LOCTEXT_NAMESPACE "MeshEditorMode"
+
+
+namespace HardenOrSoftenEdgeHelpers
+{
+	void MakeSelectedEdgesHardOrSoft( IMeshEditorModeEditingContract& MeshEditorMode, const bool bMakeEdgesHard )
+	{
+		if( MeshEditorMode.GetActiveAction() != NAME_None )
+		{
+			return;
+		}
+
+		// Make edges hard or soft
+		static TMap< UEditableMesh*, TArray<FMeshElement> > MeshesAndEdges;
+		MeshEditorMode.GetSelectedMeshesAndEdges( /* Out */ MeshesAndEdges );
+
+		if( MeshesAndEdges.Num() == 0 )
+		{
+			return;
+		}
+
+		const FScopedTransaction Transaction( bMakeEdgesHard ? LOCTEXT( "UndoHardenEdge", "Harden Edge" ) : LOCTEXT( "UndoSoftenEdge", "Soften Edge" ) );
+
+		MeshEditorMode.CommitSelectedMeshes();
+
+		// Refresh selection (committing may have created a new mesh instance)
+		MeshEditorMode.GetSelectedMeshesAndEdges( /* Out */ MeshesAndEdges );
+
+		for( auto& MeshAndEdges : MeshesAndEdges )
+		{
+			static TArray<FEdgeID> EdgeIDs;
+			EdgeIDs.Reset();
+
+			static TArray<bool> EdgesNewIsHard;
+			EdgesNewIsHard.Reset();
+
+			for( FMeshElement& EdgeElement : MeshAndEdges.Value )
+			{
+				const FEdgeID EdgeID( EdgeElement.ElementAddress.ElementID );
+
+				EdgeIDs.Add( EdgeID );
+				EdgesNewIsHard.Add( bMakeEdgesHard );
+			}
+
+
+			UEditableMesh* EditableMesh = MeshAndEdges.Key;
+			verify( !EditableMesh->AnyChangesToUndo() );
+
+			EditableMesh->StartModification( EMeshModificationType::Final, EMeshTopologyChange::TopologyChange );
+
+			EditableMesh->SetEdgesHardness( EdgeIDs, EdgesNewIsHard );
+
+			EditableMesh->EndModification();
+
+			MeshEditorMode.TrackUndo( EditableMesh, EditableMesh->MakeUndo() );
+		}
+	}
+}
+
+
+void UHardenEdgeCommand::RegisterUICommand( FBindingContext* BindingContext )
+{
+	UI_COMMAND_EXT( BindingContext, /* Out */ UICommandInfo, "HardenEdge", "Harden Edge", "Sets the edge to be hard.", EUserInterfaceActionType::Button, FInputChord( EKeys::H ) );
+}
+
+
+void UHardenEdgeCommand::Execute( IMeshEditorModeEditingContract& MeshEditorMode )
+{
+	const bool bShouldHarden = true;
+	HardenOrSoftenEdgeHelpers::MakeSelectedEdgesHardOrSoft( MeshEditorMode, bShouldHarden );
+}
+
+
+void USoftenEdgeCommand::RegisterUICommand( FBindingContext* BindingContext )
+{
+	UI_COMMAND_EXT( BindingContext, /* Out */ UICommandInfo, "SoftenEdge", "Soften Edge", "Sets the edge to be soft.", EUserInterfaceActionType::Button, FInputChord( EKeys::H, EModifierKey::Shift ) );
+}
+
+
+void USoftenEdgeCommand::Execute( IMeshEditorModeEditingContract& MeshEditorMode )
+{
+	const bool bShouldHarden = false;
+	HardenOrSoftenEdgeHelpers::MakeSelectedEdgesHardOrSoft( MeshEditorMode, bShouldHarden );
+}
+
+
+#undef LOCTEXT_NAMESPACE
+
diff --git a/Engine/Plugins/Editor/MeshEditor/Source/PolygonModeling/HardenOrSoftenEdge.h b/Engine/Plugins/Editor/MeshEditor/Source/PolygonModeling/HardenOrSoftenEdge.h
new file mode 100644
--- /dev/null
+++ b/Engine/Plugins/Editor/MeshEditor/Source/PolygonModeling/HardenOrSoftenEdge.h
@@ -0,0 +1,36 @@
+// Copyright 1998-2016 Epic Games, Inc. All Rights Reserved.
+
+#pragma once
+
+#include "MeshEditorCommands.h"
+#include "HardenOrSoftenEdge.generated.h"
+
+
+/** Makes an edge hard */
+UCLASS()
+class UHardenEdgeCommand : public UMeshEditorEdgeCommand
+{
+	GENERATED_BODY()
+
+protected:
+
+	// Overrides
+	virtual void RegisterUICommand( class FBindingContext* BindingContext ) override;
+	virtual void Execute( class IMeshEditorModeEditingContract& MeshEditorMode ) override;
+
+};
+
+
+/** Makes an edge soft */
+UCLASS()
+class USoftenEdgeCommand : public UMeshEditorEdgeCommand
+{
+	GENERATED_BODY()
+
+protected:
+
+	// Overrides
+	virtual void RegisterUICommand( class FBindingContext* BindingContext ) override;
+	virtual void Execute( class IMeshEditorModeEditingContract& MeshEditorMode ) override;
+
+};
diff --git a/Engine/Plugins/Editor/MeshEditor/Source/PolygonModeling/RemoveEdge.cpp b/Engine/Plugins/Editor/MeshEditor/Source/PolygonModeling/RemoveEdge.cpp
new file mode 100644
--- /dev/null
+++ b/Engine/Plugins/Editor/MeshEditor/Source/PolygonModeling/RemoveEdge.cpp
@@ -0,0 +1,123 @@
+// Copyright 1998-2016 Epic Games, Inc. All Rights Reserved.
+
+#include "RemoveEdge.h"
+#include "IMeshEditorModeEditingContract.h"
+#include "IMeshEditorModeUIContract.h"
+#include "UICommandInfo.h"
+#include "EditableMesh.h"
+#include "MeshElement.h"
+
+
+#define LOCTEXT_NAMESPACE "MeshEditorMode"
+
+
+void URemoveEdgeCommand::RegisterUICommand( FBindingContext* BindingContext )
+{
+	UI_COMMAND_EXT( BindingContext, /* Out */ UICommandInfo, "RemoveEdge", "Remove Edge", "Attempts to remove the selected edge and merge adjacent polygons.", EUserInterfaceActionType::Button, FInputChord( EKeys::BackSpace ) );
+}
+
+
+void URemoveEdgeCommand::Execute( IMeshEditorModeEditingContract& MeshEditorMode )
+{
+	if( MeshEditorMode.GetActiveAction() != NAME_None )
+	{
+		return;
+	}
+
+	static TMap< UEditableMesh*, TArray< FMeshElement > > MeshesWithEdgesToRemove;
+	MeshEditorMode.GetSelectedMeshesAndEdges( MeshesWithEdgesToRemove );
+
+	// @todo mesheditor: Only if one edge is selected, for now.  It gets a bit confusing when performing this operation
+	// with more than one edge selected.  However, it can be very useful when collapsing away edges that don't share any
+	// common polygons, so we should try to support it.
+	if( MeshesWithEdgesToRemove.Num() != 1 )
+	{
+		const auto& FirstMeshWithEdges = MeshesWithEdgesToRemove.CreateConstIterator();
+		if( FirstMeshWithEdges->Get<1>().Num() != 1 )
+		{
+			return;
+		}
+	}
+
+	FScopedTransaction Transaction( LOCTEXT( "UndoRemoveEdge", "Remove Edge" ) );
+
+	MeshEditorMode.CommitSelectedMeshes();
+
+	// Refresh selection (committing may have created a new mesh instance)
+	MeshEditorMode.GetSelectedMeshesAndEdges( MeshesWithEdgesToRemove );
+
+	MeshEditorMode.DeselectMeshElements( MeshesWithEdgesToRemove );
+
+	TArray<FMeshElement> MeshElementsToSelect;
+	for( const auto& MeshAndElements : MeshesWithEdgesToRemove )
+	{
+		UEditableMesh* EditableMesh = MeshAndElements.Key;
+
+		EditableMesh->StartModification( EMeshModificationType::Final, EMeshTopologyChange::TopologyChange );
+
+		const TArray<FMeshElement>& EdgeElementsToRemove = MeshAndElements.Value;
+
+		for( const FMeshElement& EdgeElementToRemove : EdgeElementsToRemove )
+		{
+			const FEdgeID EdgeID( EdgeElementToRemove.ElementAddress.ElementID );
+			{
+				bool bWasEdgeRemoved = false;
+				FPolygonRef NewPolygonRef;
+
+				EditableMesh->TryToRemovePolygonEdge( EdgeID, /* Out */ bWasEdgeRemoved, /* Out */ NewPolygonRef );
+
+				if( bWasEdgeRemoved )
+				{
+					// Select the new polygon
+					FMeshElement NewPolygonMeshElement;
+					{
+						NewPolygonMeshElement.Component = EdgeElementToRemove.Component;
+						NewPolygonMeshElement.ElementAddress = EdgeElementToRemove.ElementAddress;
+						NewPolygonMeshElement.ElementAddress.ElementType = EEditableMeshElementType::Polygon;
+						NewPolygonMeshElement.ElementAddress.SectionID = NewPolygonRef.SectionID;
+						NewPolygonMeshElement.ElementAddress.ElementID = NewPolygonRef.PolygonID;
+					}
+
+					MeshElementsToSelect.Add( NewPolygonMeshElement );
+				}
+				else
+				{
+					// Couldn't remove the edge
+					// @todo mesheditor: Needs good user feedback when this happens
+				}
+			}
+		}
+
+		EditableMesh->EndModification();
+
+		MeshEditorMode.TrackUndo( EditableMesh, EditableMesh->MakeUndo() );
+	}
+
+	{
+		// Make sure we're not still hovering over a polygon we removed
+		MeshEditorMode.ClearInvalidSelectedElements();
+
+		// Select the polygon leftover after removing the edge
+		MeshEditorMode.SelectMeshElements( MeshElementsToSelect );
+	}
+}
+
+
+void URemoveEdgeCommand::AddToVRRadialMenuActionsMenu( IMeshEditorModeUIContract& MeshEditorMode, FMenuBuilder& MenuBuilder, TSharedPtr<FUICommandList> CommandList, const FName TEMPHACK_StyleSetName, class UVREditorMode* VRMode )
+{
+	if( MeshEditorMode.GetMeshElementSelectionMode() == EEditableMeshElementType::Edge )
+	{
+		MenuBuilder.AddMenuEntry(
+			LOCTEXT( "VRRemoveEdge", "Remove" ),
+			FText(),
+			FSlateIcon( TEMPHACK_StyleSetName, "MeshEditorMode.EdgeRemove" ),
+			MakeUIAction( MeshEditorMode ),
+			NAME_None,
+			EUserInterfaceActionType::CollapsedButton
+		);
+	}
+}
+
+
+#undef LOCTEXT_NAMESPACE
+
diff --git a/Engine/Plugins/Editor/MeshEditor/Source/PolygonModeling/RemoveEdge.h b/Engine/Plugins/Editor/MeshEditor/Source/PolygonModeling/RemoveEdge.h
new file mode 100644
--- /dev/null
+++ b/Engine/Plugins/Editor/MeshEditor/Source/PolygonModeling/RemoveEdge.h
@@ -0,0 +1,22 @@
+// Copyright 1998-2016 Epic Games, Inc. All Rights Reserved.
+
+#pragma once
+
+#include "MeshEditorCommands.h"
+#include "RemoveEdge.generated.h"
+
+
+/** Attempts to remove the selected edge from the polygon, merging the adjacent polygons together */
+UCLASS()
+class URemoveEdgeCommand : public UMeshEditorEdgeCommand
+{
+	GENERATED_BODY()
+
+protected:
+
+	// Overrides
+	virtual void RegisterUICommand( class FBindingContext* BindingContext ) override;
+	virtual void Execute( class IMeshEditorModeEditingContract& MeshEditorMode ) override;
+	virtual void AddToVRRadialMenuActionsMenu( class IMeshEditorModeUIContract& MeshEditorMode, class FMenuBuilder& MenuBuilder, TSharedPtr<FUICommandList> CommandList, const FName TEMPHACK_StyleSetName, class UVREditorMode* VRMode ) override;
+
+};
diff --git a/Engine/Plugins/Editor/MeshEditor/Source/PolygonModeling/RemoveVertex.cpp b/Engine/Plugins/Editor/MeshEditor/Source/PolygonModeling/RemoveVertex.cpp
new file mode 100644
--- /dev/null
+++ b/Engine/Plugins/Editor/MeshEditor/Source/PolygonModeling/RemoveVertex.cpp
@@ -0,0 +1,121 @@
+// Copyright 1998-2016 Epic Games, Inc. All Rights Reserved.
+
+#include "RemoveVertex.h"
+#include "IMeshEditorModeEditingContract.h"
+#include "IMeshEditorModeUIContract.h"
+#include "UICommandInfo.h"
+#include "EditableMesh.h"
+#include "MeshElement.h"
+
+
+#define LOCTEXT_NAMESPACE "MeshEditorMode"
+
+
+void URemoveVertexCommand::RegisterUICommand( FBindingContext* BindingContext )
+{
+	UI_COMMAND_EXT( BindingContext, /* Out */ UICommandInfo, "RemoveVertex", "Remove Vertex", "Attempts to remove the selected vertex, keeping the polygon intact.", EUserInterfaceActionType::Button, FInputChord( EKeys::BackSpace ) );
+}
+
+
+void URemoveVertexCommand::Execute( IMeshEditorModeEditingContract& MeshEditorMode )
+{
+	if( MeshEditorMode.GetActiveAction() != NAME_None )
+	{
+		return;
+	}
+
+	static TMap< UEditableMesh*, TArray<FMeshElement> > MeshesWithVerticesToRemove;
+	MeshEditorMode.GetSelectedMeshesAndVertices( MeshesWithVerticesToRemove );
+
+	// @todo mesheditor: Only if one vertex is selected, for now.  It gets a bit confusing when performing this operation
+	// with more than one vertex selected.  However, it can be very useful when collapsing away vertices that don't share any
+	// common polygons, so we should try to support it.
+	if( MeshesWithVerticesToRemove.Num() != 1 )
+	{
+		const auto& FirstMeshWithVertices = MeshesWithVerticesToRemove.CreateConstIterator();
+		if( FirstMeshWithVertices->Get<1>().Num() != 1 )
+		{
+			return;
+		}
+	}
+
+	FScopedTransaction Transaction( LOCTEXT( "UndoRemoveVertex", "Remove Vertex" ) );
+
+	MeshEditorMode.CommitSelectedMeshes();
+
+	// Refresh selection (committing may have created a new mesh instance)
+	MeshEditorMode.GetSelectedMeshesAndVertices( MeshesWithVerticesToRemove );
+
+	// Deselect the mesh elements before we delete them.  This will make sure they become selected again after undo.
+	MeshEditorMode.DeselectMeshElements( MeshesWithVerticesToRemove );
+
+	TArray<FMeshElement> MeshElementsToSelect;
+	for( const auto& MeshAndElements : MeshesWithVerticesToRemove )
+	{
+		UEditableMesh* EditableMesh = MeshAndElements.Key;
+
+		EditableMesh->StartModification( EMeshModificationType::Final, EMeshTopologyChange::TopologyChange );
+
+		const TArray<FMeshElement>& VertexElementsToRemove = MeshAndElements.Value;
+
+		for( const FMeshElement& VertexElementToRemove : VertexElementsToRemove )
+		{
+			const FVertexID VertexID( VertexElementToRemove.ElementAddress.ElementID );
+			{
+				bool bWasVertexRemoved = false;
+				FEdgeID NewEdgeID = FEdgeID::Invalid;
+
+				EditableMesh->TryToRemoveVertex( VertexID, /* Out */ bWasVertexRemoved, /* Out */ NewEdgeID );
+
+				if( bWasVertexRemoved )
+				{
+					// Select the new edge
+					FMeshElement NewEdgeMeshElement;
+					{
+						NewEdgeMeshElement.Component = VertexElementToRemove.Component;
+						NewEdgeMeshElement.ElementAddress = VertexElementToRemove.ElementAddress;
+						NewEdgeMeshElement.ElementAddress.ElementType = EEditableMeshElementType::Edge;
+						NewEdgeMeshElement.ElementAddress.ElementID = NewEdgeID;
+					}
+
+					MeshElementsToSelect.Add( NewEdgeMeshElement );
+				}
+				else
+				{
+					// Couldn't remove the vertex
+					// @todo mesheditor: Needs good user feedback when this happens
+					// @todo mesheditor: If this fails, it will already potentially have created a new instance. To be 100% correct, it needs to do a prepass
+					// to determine whether the operation can complete successfully before actually doing it.
+				}
+			}
+		}
+
+		EditableMesh->EndModification();
+
+		MeshEditorMode.TrackUndo( EditableMesh, EditableMesh->MakeUndo() );
+	}
+
+	// Make sure we're not still hovering over a polygon we removed
+	MeshEditorMode.ClearInvalidSelectedElements();
+	MeshEditorMode.SelectMeshElements( MeshElementsToSelect );
+}
+
+
+void URemoveVertexCommand::AddToVRRadialMenuActionsMenu( IMeshEditorModeUIContract& MeshEditorMode, FMenuBuilder& MenuBuilder, TSharedPtr<FUICommandList> CommandList, const FName TEMPHACK_StyleSetName, class UVREditorMode* VRMode )
+{
+	if( MeshEditorMode.GetMeshElementSelectionMode() == EEditableMeshElementType::Vertex )
+	{
+		MenuBuilder.AddMenuEntry(
+			LOCTEXT( "VRRemoveVertex", "Remove" ),
+			FText(),
+			FSlateIcon( TEMPHACK_StyleSetName, "MeshEditorMode.VertexRemove" ),
+			MakeUIAction( MeshEditorMode ),
+			NAME_None,
+			EUserInterfaceActionType::CollapsedButton
+		);
+	}
+}
+
+
+#undef LOCTEXT_NAMESPACE
+
diff --git a/Engine/Plugins/Editor/MeshEditor/Source/PolygonModeling/RemoveVertex.h b/Engine/Plugins/Editor/MeshEditor/Source/PolygonModeling/RemoveVertex.h
new file mode 100644
--- /dev/null
+++ b/Engine/Plugins/Editor/MeshEditor/Source/PolygonModeling/RemoveVertex.h
@@ -0,0 +1,22 @@
+// Copyright 1998-2016 Epic Games, Inc. All Rights Reserved.
+
+#pragma once
+
+#include "MeshEditorCommands.h"
+#include "RemoveVertex.generated.h"
+
+
+/** Attempts to remove the selected vertex from the polygon it belongs to, keeping the polygon intact */
+UCLASS()
+class URemoveVertexCommand : public UMeshEditorVertexCommand
+{
+	GENERATED_BODY()
+
+protected:
+
+	// Overrides
+	virtual void RegisterUICommand( class FBindingContext* BindingContext ) override;
+	virtual void Execute( class IMeshEditorModeEditingContract& MeshEditorMode ) override;
+	virtual void AddToVRRadialMenuActionsMenu( class IMeshEditorModeUIContract& MeshEditorMode, class FMenuBuilder& MenuBuilder, TSharedPtr<FUICommandList> CommandList, const FName TEMPHACK_StyleSetName, class UVREditorMode* VRMode ) override;
+
+};
