Mesh Editor: Cleaned up handling of "valid" mesh elements
- We no longer have to (slowly) check for valid elements every frame while dragging
- Removed functions to deselect invalid mesh elements and to clear hover on invalid mesh elements
- Cleaned up some use of 'const' in the mesh editor API

diff --git a/Engine/Plugins/Editor/MeshEditor/Source/MeshEditor/MeshEditorMode.cpp b/Engine/Plugins/Editor/MeshEditor/Source/MeshEditor/MeshEditorMode.cpp
--- a/Engine/Plugins/Editor/MeshEditor/Source/MeshEditor/MeshEditorMode.cpp
+++ b/Engine/Plugins/Editor/MeshEditor/Source/MeshEditor/MeshEditorMode.cpp
@@ -891,20 +891,20 @@ void FMeshEditorMode::Tick( FEditorViewportClient* ViewportClient, float DeltaTi
 	// If we're currently selecting elements by painting, go ahead and do that now
 	if( ActiveAction == EMeshEditAction::SelectByPainting )
 	{
-		FMeshEditorInteractorData& MeshEditorInteractorData = GetMeshEditorInteractorData( ActiveActionInteractor );
+		const FMeshElement& HoveredMeshElement = GetHoveredMeshElement( ActiveActionInteractor );
 
 		// If not already selected, add it to our selection set
-		if( !IsMeshElementSelected( MeshEditorInteractorData.HoveredMeshElement ) )
+		if( HoveredMeshElement.IsValidMeshElement() && !IsMeshElementSelected( HoveredMeshElement ) )
 		{
 			// Only add elements of the same type.  Otherwise it would just cause things to become deselected as you move between
 			// different element types, as we don't allow you to select elements that have overlapping geometry
 			if( GetSelectedMeshElementType() == EEditableMeshElementType::Invalid || 
-				GetSelectedMeshElementType() == MeshEditorInteractorData.HoveredMeshElement.ElementAddress.ElementType ) 
+				GetSelectedMeshElementType() == HoveredMeshElement.ElementAddress.ElementType ) 
 			{
 				FSelectOrDeselectMeshElementsChangeInput ChangeInput;
 
 				// Select the element under the mouse cursor
-				ChangeInput.MeshElementsToSelect.Add( MeshEditorInteractorData.HoveredMeshElement );
+				ChangeInput.MeshElementsToSelect.Add( HoveredMeshElement );
 
 				check( MeshEditorModeProxyObject != nullptr );
 				SelectingByPaintingRevertChangeInput->Subchanges.Add( FSelectOrDeselectMeshElementsChange( MoveTemp( ChangeInput ) ).Execute( MeshEditorModeProxyObject ) );
@@ -1297,9 +1297,6 @@ bool FMeshEditorMode::DeleteSelectedMeshElement()
 		TrackUndo( EditableMesh, EditableMesh->MakeUndo() );
 	}
 
-	// Make sure we're not still hovering over a polygon we deleted
-	ClearInvalidSelectedElements();
-
 	return true;
 }
 
@@ -1624,11 +1621,7 @@ bool FMeshEditorMode::WeldSelectedVertices()
 		TrackUndo( EditableMesh, EditableMesh->MakeUndo() );
 	}
 
-	{
-		// Make sure we're not still hovering over a polygon we removed
-		ClearInvalidSelectedElements();
-		SelectMeshElements( MeshElementsToSelect );
-	}
+	SelectMeshElements( MeshElementsToSelect );
 
 	return true;
 }
@@ -1762,11 +1755,7 @@ bool FMeshEditorMode::TriangulateSelectedPolygons()
 	}
 
 	// Select the newly-created triangles
-	{
-		// Make sure we're not still hovering over a polygon we removed
-		ClearInvalidSelectedElements();
-		SelectMeshElements( MeshElementsToSelect );
-	}
+	SelectMeshElements( MeshElementsToSelect );
 
 
 	return true;
@@ -1867,11 +1856,7 @@ bool FMeshEditorMode::AssignMaterialToSelectedPolygons( UMaterialInterface* Sele
 		}
 
 		// Select the newly-created triangles
-		{
-			// Make sure we're not still hovering over a polygon we removed
-			ClearInvalidSelectedElements();
-			SelectMeshElements( MeshElementsToSelect );
-		}
+		SelectMeshElements( MeshElementsToSelect );
 
 		return true;
 	}
@@ -2421,7 +2406,11 @@ void FMeshEditorMode::Render( const FSceneView* SceneView, FViewport* Viewport,
 					if( MeshElementSelectionMode == EEditableMeshElementType::Any ||
 					    MeshEditorInteractorData.HoveredMeshElement.ElementAddress.ElementType == MeshElementSelectionMode )
 					{
-						HoveredMeshElementsToDraw.Add( MeshEditorInteractorData.HoveredMeshElement );
+						FMeshElement HoveredMeshElement = GetHoveredMeshElement( MeshEditorInteractorData.ViewportInteractor.Get() );
+						if( HoveredMeshElement.IsValidMeshElement() )
+						{
+							HoveredMeshElementsToDraw.Add( HoveredMeshElement );
+						}
 					}
 				}
 
@@ -2439,20 +2428,31 @@ void FMeshEditorMode::Render( const FSceneView* SceneView, FViewport* Viewport,
 
 			// Draw meshes that were previously hovered
 			{
+				const double CurrentRealTime = FSlateApplication::Get().GetCurrentTime();
+
+				static TArray<FMeshElement> FadingOutHoveredMeshElementsToDraw;
+				FadingOutHoveredMeshElementsToDraw.Reset();
 				static TArray<FColor> PerElementColors;
 				PerElementColors.Reset();
 
-				const double CurrentRealTime = FSlateApplication::Get().GetCurrentTime();
 
 				const float HoverFadeTime = MeshEd::HoverFadeDuration->GetFloat();;
 				for( FMeshElement& FadingOutHoveredMeshElement : FadingOutHoveredMeshElements )
 				{
-					const float TimeSinceLastHovered = CurrentRealTime - FadingOutHoveredMeshElement.LastHoverTime;
-					float Opacity = 1.0f - ( TimeSinceLastHovered / HoverFadeTime );
-					Opacity = Opacity * Opacity * Opacity * Opacity;		// Exponential falloff
-					Opacity = FMath::Clamp( Opacity, 0.0f, 1.0f );
+					if( FadingOutHoveredMeshElement.IsValidMeshElement() )
+					{
+						const UEditableMesh* EditableMesh = FindEditableMesh( *FadingOutHoveredMeshElement.Component.Get(), FadingOutHoveredMeshElement.ElementAddress.SubMeshAddress );
+						if( IsElementIDValid( FadingOutHoveredMeshElement, EditableMesh ) )
+						{
+							const float TimeSinceLastHovered = CurrentRealTime - FadingOutHoveredMeshElement.LastHoverTime;
+							float Opacity = 1.0f - ( TimeSinceLastHovered / HoverFadeTime );
+							Opacity = Opacity * Opacity * Opacity * Opacity;		// Exponential falloff
+							Opacity = FMath::Clamp( Opacity, 0.0f, 1.0f );
 
-					PerElementColors.Add( FLinearColor::Green.CopyWithNewOpacity( Opacity ).ToFColor( false ) );
+							FadingOutHoveredMeshElementsToDraw.Add( FadingOutHoveredMeshElement );
+							PerElementColors.Add( FLinearColor::Green.CopyWithNewOpacity( Opacity ).ToFColor( false ) );
+						}
+					}
 				}
 
 				const bool bFillFaces = true;
@@ -2460,7 +2460,7 @@ void FMeshEditorMode::Render( const FSceneView* SceneView, FViewport* Viewport,
 					CameraToWorld,
 					Viewport,
 					PDI,
-					FadingOutHoveredMeshElements,
+					FadingOutHoveredMeshElementsToDraw,
 					FColor::White,	// Ignored, as we'll pass in per-element colors also
 					bFillFaces,
 					HoveredSizeBias,
@@ -2501,7 +2501,7 @@ void FMeshEditorMode::Render( const FSceneView* SceneView, FViewport* Viewport,
 }
 
 
-FMeshEditorMode::FMeshEditorInteractorData& FMeshEditorMode::GetMeshEditorInteractorData( UViewportInteractor* ViewportInteractor ) const
+FMeshEditorMode::FMeshEditorInteractorData& FMeshEditorMode::GetMeshEditorInteractorData( const UViewportInteractor* ViewportInteractor ) const
 {
 	check( ViewportInteractor != nullptr );
 
@@ -2552,9 +2552,6 @@ void FMeshEditorMode::OnViewportInteractionHoverUpdate( UViewportInteractor* Vie
 		MeshEditorInteractorData.bGrabberSphereIsValid = ViewportInteractor->GetGrabberSphere( /* Out */ MeshEditorInteractorData.GrabberSphere );
 		MeshEditorInteractorData.bLaserIsValid = ViewportInteractor->GetLaserPointer( /* Out */ MeshEditorInteractorData.LaserStart, /* Out */ MeshEditorInteractorData.LaserEnd );
 
-		// @todo gizmo: Hover should have already been cleared earlier this frame when Tick() was called, but in some cases when clicking, tick won't be called
-		//check( !MeshEditorInteractorData.HoveredMeshElement.IsValidMeshElement() );
-
 		const int32 LODIndex = 0;			// @todo mesheditor: We'll want to select an LOD to edit in various different wants (LOD that's visible, or manual user select, etc.)
 											
 		// Don't use the laser pointer while if someone else has captured input
@@ -2902,9 +2899,6 @@ void FMeshEditorMode::UpdateActiveAction( const bool bIsActionFinishing )
 	}
 
 
-	bool bShouldDeselectAllFirst = false;
-	TArray<FMeshElement> MeshElementsToSelect;
-
 	// @todo mesheditor urgent: During an interactive edit, if nothing ends up selected after the edit is complete, 
 	// no mesh elements will be rendered that frame, which makes it hard to see what's going on.  Currently, we make
 	// sure that something is always selected after every type of interactive edit, but in the future that may not make sense.
@@ -3124,8 +3118,10 @@ void FMeshEditorMode::UpdateActiveAction( const bool bIsActionFinishing )
 
 				EditableMesh->CreateVertices( VerticesToCreate, NewVertexIDs );
 
+				DeselectAllMeshElements();
+
 				// Select new vertices
-				bShouldDeselectAllFirst = true;
+				TArray<FMeshElement> MeshElementsToSelect;
 				for( FVertexID VertexID : NewVertexIDs )
 				{
 					MeshElementsToSelect.Emplace(
@@ -3135,6 +3131,8 @@ void FMeshEditorMode::UpdateActiveAction( const bool bIsActionFinishing )
 					);
 				}
 
+				SelectMeshElements( MeshElementsToSelect );
+
 				if( DrawnPoints.Num() == 2 )
 				{
 					// If only two points, create an edge
@@ -3251,26 +3249,9 @@ void FMeshEditorMode::UpdateActiveAction( const bool bIsActionFinishing )
 	}
 
 
-	// Update selection state.  Note that we intentionally do this BEFORE we perform any dragging, so that
+	// Note that we intentionally make sure all selection set changes are finished  BEFORE we perform any dragging, so that
 	// we'll be dragging any newly-generated geometry from the mesh edit action.  For example, when extending
 	// an edge we want to drag around the newly-created edge, not the edge that was selected before.
-	{
-		if( bShouldDeselectAllFirst )	// @todo mesheditor extensibility: Remove this stuff ideally
-		{
-			DeselectAllMeshElements();
-		}
-		else
-		{
-			// Make sure nothing is still selected that was deleted
-			ClearInvalidSelectedElements();
-		}
-
-		if( MeshElementsToSelect.Num() > 0 )
-		{
-			SelectMeshElements( MeshElementsToSelect );
-		}
-	}
-
 	if( bIsMovingSelectedMeshElements )
 	{
 		static TMap< UEditableMesh*, TArray< const FMeshElementViewportTransformable* > > MeshesAndTransformables;
@@ -3923,93 +3904,6 @@ EEditableMeshElementType FMeshEditorMode::GetSelectedMeshElementType() const
 }
 
 
-TUniquePtr<FChange> FMeshEditorMode::ClearInvalidSelectedElements_Internal()
-{
-	FSelectOrDeselectMeshElementsChangeInput RevertInput;
-
-	for( const FMeshElement& MeshElement : SelectedMeshElements )
-	{
-		bool bIsValid = false;
-		if( MeshElement.IsValidMeshElement() )
-		{
-			const UEditableMesh* EditableMesh = FindOrCreateEditableMesh( *MeshElement.Component, MeshElement.ElementAddress.SubMeshAddress );
-			if( EditableMesh != nullptr )
-			{
-				if( IsElementIDValid( MeshElement, EditableMesh ) )
-				{
-					bIsValid = true;
-				}
-			}
-		}
-
-		if( !bIsValid )
-		{
-			// @todo mesheditor debug
-			//GWarn->Logf( TEXT( "INVALID ELEMENT deselected: %s" ), *MeshElement.ToString() );
-			RevertInput.MeshElementsToDeselect.Add( MeshElement );
-		}
-	}
-
-	return RevertInput.MeshElementsToDeselect.Num() > 0 ? ( FSelectOrDeselectMeshElementsChange( MoveTemp( RevertInput ) ).Execute( MeshEditorModeProxyObject ) ) : nullptr;
-}
-
-
-void FMeshEditorMode::ClearInvalidSelectedElements()
-{
-	// Make sure we're not still hovering over a polygon we removed
-	TUniquePtr<FChange> RevertChange = ClearInvalidSelectedElements_Internal();
-	TrackUndo( MeshEditorModeProxyObject, MoveTemp( RevertChange ) );
-}
-
-
-void FMeshEditorMode::ClearInvalidHoveredElements()
-{
-	for( FMeshEditorInteractorData& MeshEditorInteractorData : MeshEditorInteractorDatas )
-	{
-		bool bIsValid = false;
-		if( MeshEditorInteractorData.HoveredMeshElement.IsValidMeshElement() )
-		{
-			const UEditableMesh* EditableMesh = FindOrCreateEditableMesh( *MeshEditorInteractorData.HoveredMeshElement.Component, MeshEditorInteractorData.HoveredMeshElement.ElementAddress.SubMeshAddress );
-			if( EditableMesh != nullptr )
-			{
-				if( IsElementIDValid( MeshEditorInteractorData.HoveredMeshElement, EditableMesh ) )
-				{
-					bIsValid = true;
-				}
-			}
-		}
-
-		if( !bIsValid )
-		{
-			MeshEditorInteractorData.HoveredMeshElement = FMeshElement();
-		}
-	}
-
-	for( int32 MeshElementIndex = 0; MeshElementIndex < FadingOutHoveredMeshElements.Num(); ++MeshElementIndex )
-	{
-		const FMeshElement& MeshElement = FadingOutHoveredMeshElements[ MeshElementIndex ];
-
-		bool bIsValid = false;
-		if( MeshElement.IsValidMeshElement() )
-		{
-			const UEditableMesh* EditableMesh = FindOrCreateEditableMesh( *MeshElement.Component, MeshElement.ElementAddress.SubMeshAddress );
-			if( EditableMesh != nullptr )
-			{
-				if( IsElementIDValid( MeshElement, EditableMesh ) )
-				{
-					bIsValid = true;
-				}
-			}
-		}
-
-		if( !bIsValid )
-		{
-			FadingOutHoveredMeshElements.RemoveAt( MeshElementIndex-- );
-		}
-	}
-}
-
-
 void FMeshEditorMode::OnViewportInteractionInputAction( FEditorViewportClient& ViewportClient, UViewportInteractor* ViewportInteractor, const FViewportActionKeyInput& Action, bool& bOutIsInputCaptured, bool& bWasHandled )
 {
 	if( !bWasHandled && Action.ActionType == ViewportWorldActionTypes::SelectAndMove )
@@ -4046,29 +3940,31 @@ void FMeshEditorMode::OnViewportInteractionInputAction( FEditorViewportClient& V
 				bWasHandled = true;
 			}
 
-			else if( MeshEditorInteractorData.HoveredMeshElement.IsValidMeshElement() && 
+			else if( GetHoveredMeshElement( MeshEditorInteractorData.ViewportInteractor.Get() ).IsValidMeshElement() && 
 				( MeshEditorInteractorData.bLaserIsValid || MeshEditorInteractorData.bGrabberSphereIsValid ) )
 			{
+				FMeshElement HoveredMeshElement = GetHoveredMeshElement( MeshEditorInteractorData.ViewportInteractor.Get() );
+
 				// Make sure the actor is selected
 				// @todo mesheditor: Do we need/want to automatically select actors when doing mesh editing?  If so, consider how undo will 
 				// encapsulate the actor selection change with the mesh element selection change
 				if( false )
 				{
-					UPrimitiveComponent* Component = MeshEditorInteractorData.HoveredMeshElement.Component.Get();
+					UPrimitiveComponent* Component = HoveredMeshElement.Component.Get();
 					if( Component == nullptr || !GEditor->GetSelectedActors()->IsSelected( Component->GetOwner() ) )
 					{
 						GEditor->SelectNone( true, true );
 					}
 					else
 					{
-						GEditor->SelectActor( MeshEditorInteractorData.HoveredMeshElement.Component->GetOwner(), true, true );
+						GEditor->SelectActor( HoveredMeshElement.Component->GetOwner(), true, true );
 					}
 				}
 
 				// Holding down Control enables multi-select (adds to selection, or deselects single elements when already selected)
 				const bool bIsMultiSelecting = ViewportInteractor->IsModifierPressed();
 
-				const int32 AlreadySelectedMeshElement = GetSelectedMeshElementIndex( MeshEditorInteractorData.HoveredMeshElement );
+				const int32 AlreadySelectedMeshElement = GetSelectedMeshElementIndex( HoveredMeshElement );
 				if( AlreadySelectedMeshElement != INDEX_NONE && !bIsMultiSelecting )
 				{
 					const EEditableMeshElementType SelectedMeshElementType = GetSelectedMeshElementType();
@@ -4149,7 +4045,7 @@ void FMeshEditorMode::OnViewportInteractionInputAction( FEditorViewportClient& V
 						ChangeInput.MeshElementsToDeselect.Add( SelectedMeshElements[ AlreadySelectedMeshElement ] );
 						TrackUndo( MeshEditorModeProxyObject, FSelectOrDeselectMeshElementsChange( MoveTemp( ChangeInput ) ).Execute( MeshEditorModeProxyObject ) );
 					}
-					else if( MeshElementSelectionMode == EEditableMeshElementType::Any || MeshElementSelectionMode == MeshEditorInteractorData.HoveredMeshElement.ElementAddress.ElementType )
+					else if( MeshElementSelectionMode == EEditableMeshElementType::Any || MeshElementSelectionMode == HoveredMeshElement.ElementAddress.ElementType )
 					{
 						// Start painting selection
 						const bool bIsSelectByPaintingEnabled = MeshEd::EnableSelectByPainting->GetInt() != 0;
@@ -4169,7 +4065,7 @@ void FMeshEditorMode::OnViewportInteractionInputAction( FEditorViewportClient& V
 						}
 
 						// Select the element under the mouse cursor
-						ChangeInput.MeshElementsToSelect.Add( MeshEditorInteractorData.HoveredMeshElement );
+						ChangeInput.MeshElementsToSelect.Add( HoveredMeshElement );
 
 						TUniquePtr<FChange> RevertChange = FSelectOrDeselectMeshElementsChange( MoveTemp( ChangeInput ) ).Execute( MeshEditorModeProxyObject );
 
@@ -4321,9 +4217,6 @@ void FMeshEditorMode::FinishAction()
 
 void FMeshEditorMode::PostUndo()
 {
-	// Make sure we're not hovering over anything that was deleted after an undo action
-	ClearInvalidHoveredElements();
-
 	// Update our transformable list
 	RefreshTransformables();
 }
@@ -5035,7 +4928,7 @@ void FMeshEditorMode::TrackUndo( UObject* Object, TUniquePtr<FChange> RevertChan
 }
 
 
-FMeshElement FMeshEditorMode::GetHoveredMeshElement( UViewportInteractor* ViewportInteractor ) const
+FMeshElement FMeshEditorMode::GetHoveredMeshElement( const UViewportInteractor* ViewportInteractor ) const
 {
 	FMeshElement HoveredMeshElement;
 
diff --git a/Engine/Plugins/Editor/MeshEditor/Source/MeshEditor/MeshEditorMode.h b/Engine/Plugins/Editor/MeshEditor/Source/MeshEditor/MeshEditorMode.h
--- a/Engine/Plugins/Editor/MeshEditor/Source/MeshEditor/MeshEditorMode.h
+++ b/Engine/Plugins/Editor/MeshEditor/Source/MeshEditor/MeshEditorMode.h
@@ -109,7 +109,7 @@ public:
 	{
 		bool bIsValid = false;
 
-		if( MeshElement.ElementAddress.ElementID != FElementID::Invalid )
+		if( EditableMesh != nullptr && MeshElement.ElementAddress.ElementID != FElementID::Invalid )
 		{
 			switch( MeshElement.ElementAddress.ElementType )
 			{
@@ -171,7 +171,7 @@ protected:
 	virtual FName GetActiveAction() const override { return ActiveAction; }
 	virtual void TrackUndo( UObject* Object, TUniquePtr<FChange> RevertChange ) override;
 	virtual void CommitSelectedMeshes() override;
-	virtual FMeshElement GetHoveredMeshElement( class UViewportInteractor* ViewportInteractor ) const override;
+	virtual FMeshElement GetHoveredMeshElement( const class UViewportInteractor* ViewportInteractor ) const override;
 	virtual bool IsMeshElementSelected( const FMeshElement MeshElement ) const override
 	{
 		return GetSelectedMeshElementIndex( MeshElement ) != INDEX_NONE;
@@ -192,7 +192,6 @@ protected:
 	virtual void DeselectAllMeshElements() override;
 	virtual void DeselectMeshElements( const TArray<FMeshElement>& MeshElementsToDeselect ) override;
 	virtual void DeselectMeshElements( const TMap<UEditableMesh*, TArray<FMeshElement>>& MeshElementsToDeselect ) override;
-	virtual void ClearInvalidSelectedElements() override;
 	virtual void FindEdgeSplitUnderInteractor( UViewportInteractor* ViewportInteractor, const UEditableMesh* EditableMesh, const TArray<FMeshElement>& EdgeElements, TArray<float>& OutSplits ) override;
 	virtual class UViewportInteractor* GetActiveActionInteractor() override
 	{
@@ -252,9 +251,6 @@ protected:
 	/** Returns the index of an element in the selection set, or INDEX_NONE if its not selected */
 	int32 GetSelectedMeshElementIndex( const FMeshElement& MeshElement ) const;
 
-	/** Clears hover on mesh elements that may no longer be valid.  You'll want to call this if you change the mesh topology */
-	void ClearInvalidHoveredElements();
-
 	/** Updates the current view location, from either the viewport client or the VR interface, whichever is in use */
 	void UpdateCameraToWorldTransform( const FEditorViewportClient& ViewportClient );
 
@@ -281,8 +277,6 @@ protected:
 	/** Return the CommandList pertinent to the currently selected element type, or nullptr if nothing is selected */
 	const FUICommandList* GetCommandListForSelectedElementType() const;
 
-	TUniquePtr<FChange> ClearInvalidSelectedElements_Internal();
-
 	/** Commits the mesh instance for the given component */
 	void CommitEditableMeshIfNecessary( UEditableMesh* EditableMesh, UPrimitiveComponent* Component );
 
@@ -321,7 +315,7 @@ protected:
 
 	/** Gets mesh editor interactor data for the specified viewport interactor.  If we've never seen this viewport interactor before,
 	    new (empty) data will be created for it on demand */
-	FMeshEditorInteractorData& GetMeshEditorInteractorData( UViewportInteractor* ViewportInteractor ) const;
+	FMeshEditorInteractorData& GetMeshEditorInteractorData( const UViewportInteractor* ViewportInteractor ) const;
 
 	/** Helper function that returns a map keying an editable mesh with its selected elements */
 	void GetSelectedMeshesAndElements( EEditableMeshElementType ElementType, TMap<UEditableMesh*, TArray<FMeshElement>>& OutMeshesAndElements );
diff --git a/Engine/Plugins/Editor/MeshEditor/Source/MeshEditor/Public/IMeshEditorModeEditingContract.h b/Engine/Plugins/Editor/MeshEditor/Source/MeshEditor/Public/IMeshEditorModeEditingContract.h
--- a/Engine/Plugins/Editor/MeshEditor/Source/MeshEditor/Public/IMeshEditorModeEditingContract.h
+++ b/Engine/Plugins/Editor/MeshEditor/Source/MeshEditor/Public/IMeshEditorModeEditingContract.h
@@ -26,7 +26,7 @@ public:
 	virtual bool IsMeshElementSelected( const FMeshElement MeshElement ) const = 0;
 
 	/** Returns the mesh element the specified viewport interactor is currently hovering over.  The returned element might be invalid if nothing valid is hovered right now */
-	virtual FMeshElement GetHoveredMeshElement( class UViewportInteractor* ViewportInteractor ) const = 0;
+	virtual FMeshElement GetHoveredMeshElement( const class UViewportInteractor* ViewportInteractor ) const = 0;
 
 	virtual void GetSelectedMeshesAndVertices( TMap<class UEditableMesh*, TArray<FMeshElement>>& OutMeshesAndVertices ) = 0;
 	virtual void GetSelectedMeshesAndEdges( TMap<class UEditableMesh*, TArray<FMeshElement>>& OutMeshesAndEdges ) = 0;
@@ -45,9 +45,6 @@ public:
 	/** Commits all selected meshes */
 	virtual void CommitSelectedMeshes() = 0;
 
-	/** Clears hover and selection on mesh elements that may no longer be valid.  You'll want to call this if you change the mesh topology */
-	virtual void ClearInvalidSelectedElements() = 0;
-
 	/** Given an interactor and a mesh, finds edges under the interactor along with their exact split position (progress along the edge) */
 	virtual void FindEdgeSplitUnderInteractor( class UViewportInteractor* ViewportInteractor, const UEditableMesh* EditableMesh, const TArray<FMeshElement>& EdgeElements, TArray<float>& OutSplits ) = 0;
 
diff --git a/Engine/Plugins/Editor/MeshEditor/Source/PolygonModeling/ExtrudePolygon.cpp b/Engine/Plugins/Editor/MeshEditor/Source/PolygonModeling/ExtrudePolygon.cpp
--- a/Engine/Plugins/Editor/MeshEditor/Source/PolygonModeling/ExtrudePolygon.cpp
+++ b/Engine/Plugins/Editor/MeshEditor/Source/PolygonModeling/ExtrudePolygon.cpp
@@ -81,7 +81,7 @@ bool UExtrudePolygonCommand::TryStartingToDrag( IMeshEditorModeEditingContract&
 	bool bHaveExtrudeAxis = false;
 	{
 		const FMeshElement& PolygonElement = MeshEditorMode.GetHoveredMeshElement( ViewportInteractor );
-		if( PolygonElement.ElementAddress.ElementType == EEditableMeshElementType::Polygon )
+		if( PolygonElement.IsValidMeshElement() && PolygonElement.ElementAddress.ElementType == EEditableMeshElementType::Polygon )
 		{
 			// Is it selected?
 			if( MeshEditorMode.IsMeshElementSelected( PolygonElement ) )
diff --git a/Engine/Plugins/Editor/MeshEditor/Source/PolygonModeling/RemoveEdge.cpp b/Engine/Plugins/Editor/MeshEditor/Source/PolygonModeling/RemoveEdge.cpp
--- a/Engine/Plugins/Editor/MeshEditor/Source/PolygonModeling/RemoveEdge.cpp
+++ b/Engine/Plugins/Editor/MeshEditor/Source/PolygonModeling/RemoveEdge.cpp
@@ -93,13 +93,8 @@ void URemoveEdgeCommand::Execute( IMeshEditorModeEditingContract& MeshEditorMode
 		MeshEditorMode.TrackUndo( EditableMesh, EditableMesh->MakeUndo() );
 	}
 
-	{
-		// Make sure we're not still hovering over a polygon we removed
-		MeshEditorMode.ClearInvalidSelectedElements();
-
-		// Select the polygon leftover after removing the edge
-		MeshEditorMode.SelectMeshElements( MeshElementsToSelect );
-	}
+	// Select the polygon leftover after removing the edge
+	MeshEditorMode.SelectMeshElements( MeshElementsToSelect );
 }
 
 
diff --git a/Engine/Plugins/Editor/MeshEditor/Source/PolygonModeling/RemoveVertex.cpp b/Engine/Plugins/Editor/MeshEditor/Source/PolygonModeling/RemoveVertex.cpp
--- a/Engine/Plugins/Editor/MeshEditor/Source/PolygonModeling/RemoveVertex.cpp
+++ b/Engine/Plugins/Editor/MeshEditor/Source/PolygonModeling/RemoveVertex.cpp
@@ -95,8 +95,6 @@ void URemoveVertexCommand::Execute( IMeshEditorModeEditingContract& MeshEditorMo
 		MeshEditorMode.TrackUndo( EditableMesh, EditableMesh->MakeUndo() );
 	}
 
-	// Make sure we're not still hovering over a polygon we removed
-	MeshEditorMode.ClearInvalidSelectedElements();
 	MeshEditorMode.SelectMeshElements( MeshElementsToSelect );
 }
 
diff --git a/Engine/Plugins/Editor/MeshEditor/Source/PolygonModeling/TessellatePolygon.cpp b/Engine/Plugins/Editor/MeshEditor/Source/PolygonModeling/TessellatePolygon.cpp
--- a/Engine/Plugins/Editor/MeshEditor/Source/PolygonModeling/TessellatePolygon.cpp
+++ b/Engine/Plugins/Editor/MeshEditor/Source/PolygonModeling/TessellatePolygon.cpp
@@ -86,9 +86,6 @@ void UTessellatePolygonCommand::Execute( IMeshEditorModeEditingContract& MeshEdi
 		MeshEditorMode.TrackUndo( EditableMesh, EditableMesh->MakeUndo() );
 	}
 
-	// Make sure we're not still hovering over a polygon we removed
-	MeshEditorMode.ClearInvalidSelectedElements();
-
 	// Select the new smaller polygons
 	MeshEditorMode.SelectMeshElements( MeshElementsToSelect );
 }
