Mesh Editor: Fixed crash when beveling/insetting polygons
- Mesh elements were deselected twice, which caused them to be reselected twice on Redo while dragging.  The next frame, we'd try to bevel/delete the same polygon twice. Crash!
- This is fallout from an earlier change this week to simplify how selection is handled with mesh editor

diff --git a/Engine/Plugins/Editor/MeshEditor/Source/PolygonModeling/BevelOrInsetPolygon.cpp b/Engine/Plugins/Editor/MeshEditor/Source/PolygonModeling/BevelOrInsetPolygon.cpp
--- a/Engine/Plugins/Editor/MeshEditor/Source/PolygonModeling/BevelOrInsetPolygon.cpp
+++ b/Engine/Plugins/Editor/MeshEditor/Source/PolygonModeling/BevelOrInsetPolygon.cpp
@@ -166,9 +166,6 @@ namespace BevelOrInsetPolygonHelpers
 					if( InsetFixedDistance > SMALL_NUMBER ||
 						InsetProgressTowardCenter > SMALL_NUMBER )
 					{
-						// Deselect the mesh elements before we delete them.  This will make sure they become selected again after undo.
-						MeshEditorMode.DeselectMeshElements( MeshesWithPolygonsToInset );
-
 						for( auto& MeshAndPolygons : MeshesWithPolygonsToInset )
 						{
 							UEditableMesh* EditableMesh = MeshAndPolygons.Key;
