Mesh Editor: Fixed transform gizmo vanishing and reappearing unexpectedly when dragging geometry

diff --git a/Engine/Plugins/Editor/MeshEditor/Source/MeshEditor/MeshEditorMode.cpp b/Engine/Plugins/Editor/MeshEditor/Source/MeshEditor/MeshEditorMode.cpp
--- a/Engine/Plugins/Editor/MeshEditor/Source/MeshEditor/MeshEditorMode.cpp
+++ b/Engine/Plugins/Editor/MeshEditor/Source/MeshEditor/MeshEditorMode.cpp
@@ -154,7 +154,8 @@ TUniquePtr<FChange> FMeshEditorMode::FSelectOrDeselectMeshElementsChange::Execut
 		}
 
 		// Update our transformable list
-		MeshEditorMode.RefreshTransformables();
+		const bool bNewObjectsSelected = true;
+		MeshEditorMode.RefreshTransformables( bNewObjectsSelected );
 	}
 
 	return MakeUnique<FCompoundChange>( MoveTemp( CompoundRevertInput ) );
@@ -182,7 +183,9 @@ TUniquePtr<FChange> FMeshEditorMode::FDeselectAllMeshElementsChange::Execute( UO
 	if( MeshEditorMode.IsActive() )
 	{
 		MeshEditorMode.SelectedMeshElements.Reset();
-		MeshEditorMode.RefreshTransformables();
+
+		const bool bNewObjectsSelected = true;
+		MeshEditorMode.RefreshTransformables( bNewObjectsSelected );
 	}
 
 	return ( RevertInput.MeshElementsToSelect.Num() > 0 ) ? MakeUnique <FSelectOrDeselectMeshElementsChange>( MoveTemp( RevertInput ) ) : nullptr;
@@ -482,7 +485,8 @@ void FMeshEditorMode::RemoveEditableMeshReferences()
 
 	if( ViewportWorldInteraction != nullptr )
 	{
-		RefreshTransformables();
+		const bool bNewObjectsSelected = false;
+		RefreshTransformables( bNewObjectsSelected );
 	}
 }
 
@@ -1104,7 +1108,8 @@ void FMeshEditorMode::CommitEditableMeshIfNecessary( UEditableMesh* EditableMesh
 		// 			}
 		// 		}
 
-		RefreshTransformables();
+		const bool bNewObjectsSelected = false;
+		RefreshTransformables( bNewObjectsSelected );
 	}
 	else if( !EditableMesh->IsCommitted() )
 	{
@@ -4210,7 +4215,8 @@ void FMeshEditorMode::FinishAction()
 	// If the action has finished, make sure the gizmo is in the correct place as elements may have moved.
 	if( bIsActionFinishing )
 	{
-		RefreshTransformables();
+		const bool bNewObjectsSelected = false;
+		RefreshTransformables( bNewObjectsSelected );
 	}
 }
 
@@ -4218,7 +4224,8 @@ void FMeshEditorMode::FinishAction()
 void FMeshEditorMode::PostUndo()
 {
 	// Update our transformable list
-	RefreshTransformables();
+	const bool bNewObjectsSelected = false;
+	RefreshTransformables( bNewObjectsSelected );
 }
 
 
@@ -4555,7 +4562,7 @@ void FMeshEditorMode::PerformMarqueeSelect( EEditableMeshElementType ElementType
 }
 
 
-void FMeshEditorMode::RefreshTransformables()
+void FMeshEditorMode::RefreshTransformables( const bool bNewObjectsSelected )
 {
 	// Don't refresh transformables while we're actively moving them around
 	const bool bAllowRefresh =
@@ -4648,7 +4655,7 @@ void FMeshEditorMode::RefreshTransformables()
 		}
 	}
 
-	ViewportWorldInteraction->SetTransformables( MoveTemp( Transformables ) );
+	ViewportWorldInteraction->SetTransformables( MoveTemp( Transformables ), bNewObjectsSelected );
 }
 
 
diff --git a/Engine/Plugins/Editor/MeshEditor/Source/MeshEditor/MeshEditorMode.h b/Engine/Plugins/Editor/MeshEditor/Source/MeshEditor/MeshEditorMode.h
--- a/Engine/Plugins/Editor/MeshEditor/Source/MeshEditor/MeshEditorMode.h
+++ b/Engine/Plugins/Editor/MeshEditor/Source/MeshEditor/MeshEditorMode.h
@@ -324,7 +324,7 @@ protected:
 	void PerformMarqueeSelect( EEditableMeshElementType ElementType );
 
 	/** Rebuilds the list of mesh element transformables and updates the world viewport interaction system with the new list */
-	void RefreshTransformables();
+	void RefreshTransformables( const bool bNewObjectsSelected );
 
 	/** Callback when PIE/SIE ends */
 	void OnEndPIE( bool bIsSimulating );
diff --git a/Engine/Source/Editor/ViewportInteraction/ActorTransformer.cpp b/Engine/Source/Editor/ViewportInteraction/ActorTransformer.cpp
--- a/Engine/Source/Editor/ViewportInteraction/ActorTransformer.cpp
+++ b/Engine/Source/Editor/ViewportInteraction/ActorTransformer.cpp
@@ -59,7 +59,7 @@ void UActorTransformer::OnStopDragging(class UViewportInteractor* Interactor)
 	}
 }
 
-void UActorTransformer::UpdateTransformables()
+void UActorTransformer::OnActorSelectionChanged( UObject* ChangedObject )
 {
 	TArray<TUniquePtr<FViewportTransformable>> NewTransformables;
 
@@ -87,12 +87,7 @@ void UActorTransformer::UpdateTransformables()
 		}
 	}
 
-	ViewportWorldInteraction->SetTransformables( MoveTemp( NewTransformables ) );
-}
-
-
-void UActorTransformer::OnActorSelectionChanged( UObject* ChangedObject )
-{
-	UpdateTransformables();
+	const bool bNewObjectsSelected = true;
+	ViewportWorldInteraction->SetTransformables( MoveTemp( NewTransformables ), bNewObjectsSelected );
 }
 
diff --git a/Engine/Source/Editor/ViewportInteraction/Public/ActorTransformer.h b/Engine/Source/Editor/ViewportInteraction/Public/ActorTransformer.h
--- a/Engine/Source/Editor/ViewportInteraction/Public/ActorTransformer.h
+++ b/Engine/Source/Editor/ViewportInteraction/Public/ActorTransformer.h
@@ -28,8 +28,5 @@ protected:
 	/** Called when selection changes in the level */
 	void OnActorSelectionChanged( UObject* ChangedObject );
 
-	/** Creates a list of transformables with the current selected actors */
-	void UpdateTransformables();
-
 };
 
diff --git a/Engine/Source/Editor/ViewportInteraction/Public/ViewportWorldInteraction.h b/Engine/Source/Editor/ViewportInteraction/Public/ViewportWorldInteraction.h
--- a/Engine/Source/Editor/ViewportInteraction/Public/ViewportWorldInteraction.h
+++ b/Engine/Source/Editor/ViewportInteraction/Public/ViewportWorldInteraction.h
@@ -115,7 +115,7 @@ public:
 
 	/** Sets the list of objects that this system will be responsible for transforming when interacting using the
 	    gizmo or directly on the objects */
-	void SetTransformables( TArray< TUniquePtr< FViewportTransformable > >&& NewTransformables );
+	void SetTransformables( TArray< TUniquePtr< FViewportTransformable > >&& NewTransformables, const bool bNewObjectsSelected );
 
 	/** When using VR, this sets the viewport client that's been "possessed" by the head mounted display.  Only valid when VR is enabled. */
 	void SetDefaultOptionalViewportClient( const TSharedPtr<class FEditorViewportClient>& InEditorViewportClient );
diff --git a/Engine/Source/Editor/ViewportInteraction/ViewportWorldInteraction.cpp b/Engine/Source/Editor/ViewportInteraction/ViewportWorldInteraction.cpp
--- a/Engine/Source/Editor/ViewportInteraction/ViewportWorldInteraction.cpp
+++ b/Engine/Source/Editor/ViewportInteraction/ViewportWorldInteraction.cpp
@@ -486,7 +486,8 @@ void UViewportWorldInteraction::RemoveInteractor( UViewportInteractor* Interacto
 void UViewportWorldInteraction::SetTransformer( UViewportTransformer* NewTransformer )
 {
 	// Clear all existing transformables
-	SetTransformables( TArray< TUniquePtr< class FViewportTransformable > >() );
+	const bool bNewObjectsSelected = false;
+	SetTransformables( TArray< TUniquePtr< class FViewportTransformable > >(), bNewObjectsSelected );
 
 	if( ViewportTransformer != nullptr )
 	{
@@ -506,7 +507,7 @@ void UViewportWorldInteraction::SetTransformer( UViewportTransformer* NewTransfo
 }
 
 
-void UViewportWorldInteraction::SetTransformables( TArray< TUniquePtr< class FViewportTransformable > >&& NewTransformables )
+void UViewportWorldInteraction::SetTransformables( TArray< TUniquePtr< class FViewportTransformable > >&& NewTransformables, const bool bNewObjectsSelected )
 {
 	// We're dealing with a new set of transformables, so clear out anything that was going on with the old ones.
 	bDraggedSinceLastSelection = false;
@@ -526,7 +527,6 @@ void UViewportWorldInteraction::SetTransformables( TArray< TUniquePtr< class FVi
 
 	Transformables = MoveTemp( NewTransformables );
 
-	const bool bNewObjectsSelected = true;
 	RefreshTransformGizmo( bNewObjectsSelected );
 }
 
