Mesh Editor: Edit meshes without having elements selected
- You can now invoke 'common' actions without having a mesh element selected! (e.g., quadrangulate, add/remove subD)
- Blue wireframe preview now draws for selected actors even when no element is selected
- Some refactoring to how we track selection state in Mesh Editor.  We now keep track of selected components and their editable meshes.  This simplified various functions.
- We now only support editing meshes for selected actors/components.  Removed the preference to edit non-selected objects.

diff --git a/Engine/Plugins/Editor/MeshEditor/Source/MeshEditor/MeshEditorCommands.cpp b/Engine/Plugins/Editor/MeshEditor/Source/MeshEditor/MeshEditorCommands.cpp
--- a/Engine/Plugins/Editor/MeshEditor/Source/MeshEditor/MeshEditorCommands.cpp
+++ b/Engine/Plugins/Editor/MeshEditor/Source/MeshEditor/MeshEditorCommands.cpp
@@ -52,7 +52,6 @@ FMeshEditorCommonCommands::FMeshEditorCommonCommands()
 
 void FMeshEditorCommonCommands::RegisterCommands()
 {
-	UI_COMMAND(DeleteMeshElement, "Delete", "Delete selected mesh elements, including polygons partly defined by selected elements.", EUserInterfaceActionType::Button, FInputChord(EKeys::Delete));
 	UI_COMMAND(AddSubdivisionLevel, "Add Subdivision Level", "Increases the number of subdivision levels for the selected mesh.", EUserInterfaceActionType::Button, FInputChord(EKeys::Equals));
 	UI_COMMAND(RemoveSubdivisionLevel, "Remove Subdivision Level", "Decreases the number of subdivision levels for the selected mesh.", EUserInterfaceActionType::Button, FInputChord(EKeys::Hyphen));
 	UI_COMMAND(ShowVertexNormals, "Show Vertex Normals", "Toggles debug rendering of vertex normals.", EUserInterfaceActionType::ToggleButton, FInputChord(EKeys::V));
@@ -68,6 +67,21 @@ void FMeshEditorCommonCommands::RegisterCommands()
 	UI_COMMAND(QuadrangulateMesh, "Quadrangulate Mesh", "Quadrangulates the selected mesh.", EUserInterfaceActionType::Button, FInputChord());
 }
 
+FMeshEditorAnyElementCommands::FMeshEditorAnyElementCommands() 
+	: TCommands<FMeshEditorAnyElementCommands>(
+		"MeshEditorAnyElement",
+		LOCTEXT( "MeshEditorGeneral", "Mesh Editor Any Element Type" ),
+		"MainFrame",
+		FMeshEditorStyle::GetStyleSetName()	)
+{
+}
+
+
+void FMeshEditorAnyElementCommands::RegisterCommands()
+{
+	UI_COMMAND(DeleteMeshElement, "Delete", "Delete selected mesh elements, including polygons partly defined by selected elements.", EUserInterfaceActionType::Button, FInputChord(EKeys::Delete));
+}
+
 FMeshEditorVertexCommands::FMeshEditorVertexCommands() 
 	: TCommands<FMeshEditorVertexCommands>(
 		"MeshEditorVertex",
diff --git a/Engine/Plugins/Editor/MeshEditor/Source/MeshEditor/MeshEditorMode.cpp b/Engine/Plugins/Editor/MeshEditor/Source/MeshEditor/MeshEditorMode.cpp
--- a/Engine/Plugins/Editor/MeshEditor/Source/MeshEditor/MeshEditorMode.cpp
+++ b/Engine/Plugins/Editor/MeshEditor/Source/MeshEditor/MeshEditorMode.cpp
@@ -153,6 +153,8 @@ TUniquePtr<FChange> FMeshEditorMode::FSelectOrDeselectMeshElementsChange::Execut
 			}
 		}
 
+		MeshEditorMode.UpdateSelectedEditableMeshes();
+
 		// Update our transformable list
 		const bool bNewObjectsSelected = true;
 		MeshEditorMode.RefreshTransformables( bNewObjectsSelected );
@@ -184,6 +186,8 @@ TUniquePtr<FChange> FMeshEditorMode::FDeselectAllMeshElementsChange::Execute( UO
 	{
 		MeshEditorMode.SelectedMeshElements.Reset();
 
+		MeshEditorMode.UpdateSelectedEditableMeshes();
+
 		const bool bNewObjectsSelected = true;
 		MeshEditorMode.RefreshTransformables( bNewObjectsSelected );
 	}
@@ -415,6 +419,7 @@ FMeshEditorMode::FMeshEditorMode()
 
 	// Register mesh editor actions
 	FMeshEditorCommonCommands::Register();
+	FMeshEditorAnyElementCommands::Register();
 	FMeshEditorVertexCommands::Register();
 	FMeshEditorEdgeCommands::Register();
 	FMeshEditorPolygonCommands::Register();
@@ -446,6 +451,7 @@ FMeshEditorMode::~FMeshEditorMode()
 	FMeshEditorPolygonCommands::Unregister();
 	FMeshEditorEdgeCommands::Unregister();
 	FMeshEditorVertexCommands::Unregister();
+	FMeshEditorAnyElementCommands::Unregister();
 	FMeshEditorCommonCommands::Unregister();
 
 	MeshEditorModeProxyObject = nullptr;
@@ -476,6 +482,8 @@ void FMeshEditorMode::RemoveEditableMeshReferences()
 	// Instanced meshes live within the level itself. So remove all possible references to any editable mesh when the map changes,
 	// to prevent unreachable paths following GC.
 	CachedEditableMeshes.Empty();
+	SelectedComponentsAndEditableMeshes.Reset();
+	SelectedEditableMeshes.Reset();
 	SelectedMeshElements.Empty();
 	SelectedVertices.Empty();
 	SelectedEdges.Empty();
@@ -520,48 +528,51 @@ void FMeshEditorMode::PlayFinishActionSound( FName NewAction, UViewportInteracto
 
 void FMeshEditorMode::BindCommands()
 {
-	const FMeshEditorCommonCommands& MeshEditorCommonActions( FMeshEditorCommonCommands::Get() );
-	const FMeshEditorVertexCommands& MeshEditorVertexActions( FMeshEditorVertexCommands::Get() );
-	const FMeshEditorEdgeCommands& MeshEditorEdgeActions( FMeshEditorEdgeCommands::Get() );
-	const FMeshEditorPolygonCommands& MeshEditorPolygonActions( FMeshEditorPolygonCommands::Get() );
+	const FMeshEditorCommonCommands& MeshEditorCommonCommands( FMeshEditorCommonCommands::Get() );
+	const FMeshEditorAnyElementCommands& MeshEditorAnyElementCommands( FMeshEditorAnyElementCommands::Get() );
+	const FMeshEditorVertexCommands& MeshEditorVertexCommands( FMeshEditorVertexCommands::Get() );
+	const FMeshEditorEdgeCommands& MeshEditorEdgeCommands( FMeshEditorEdgeCommands::Get() );
+	const FMeshEditorPolygonCommands& MeshEditorPolygonCommands( FMeshEditorPolygonCommands::Get() );
 
 	// Register editing modes (equipped actions)
-	RegisterVertexEditingMode( MeshEditorVertexActions.MoveVertex, EMeshEditAction::Move );
+	RegisterVertexEditingMode( MeshEditorVertexCommands.MoveVertex, EMeshEditAction::Move );
 
-	RegisterEdgeEditingMode( MeshEditorEdgeActions.MoveEdge, EMeshEditAction::Move );
+	RegisterEdgeEditingMode( MeshEditorEdgeCommands.MoveEdge, EMeshEditAction::Move );
 
-	RegisterPolygonEditingMode( MeshEditorPolygonActions.MovePolygon, EMeshEditAction::Move );
-	RegisterCommonEditingMode( MeshEditorCommonActions.DrawVertices, EMeshEditAction::DrawVertices );
+	RegisterPolygonEditingMode( MeshEditorPolygonCommands.MovePolygon, EMeshEditAction::Move );
+	RegisterCommonEditingMode( MeshEditorCommonCommands.DrawVertices, EMeshEditAction::DrawVertices );
 
 	// Register commands which work regardless of which element type is selected
-	RegisterCommand( MeshEditorCommonActions.DeleteMeshElement, FExecuteAction::CreateLambda( [this] { DeleteSelectedMeshElement(); } ) );
-	RegisterCommand( MeshEditorCommonActions.AddSubdivisionLevel, FExecuteAction::CreateLambda( [this] { AddOrRemoveSubdivisionLevel( true ); } ) );
-	RegisterCommand( MeshEditorCommonActions.RemoveSubdivisionLevel, FExecuteAction::CreateLambda( [this] { AddOrRemoveSubdivisionLevel( false ); } ) );
+	RegisterAnyElementCommand( MeshEditorAnyElementCommands.DeleteMeshElement, FExecuteAction::CreateLambda( [this] { DeleteSelectedMeshElement(); } ) );
+
+	// Register commands which work even without a selected element, as long as at least one mesh is selected
+	RegisterCommonCommand( MeshEditorCommonCommands.AddSubdivisionLevel, FExecuteAction::CreateLambda( [this] { AddOrRemoveSubdivisionLevel( true ); } ) );
+	RegisterCommonCommand( MeshEditorCommonCommands.RemoveSubdivisionLevel, FExecuteAction::CreateLambda( [this] { AddOrRemoveSubdivisionLevel( false ); } ) );
 
 	// @todo mesheditor: support EUserInterfaceActionType::ToggleButton actions in the UI, and extend RegisterCommand to allow
 	// a delegate returning check state.
-	RegisterCommand( MeshEditorCommonActions.ShowVertexNormals, FExecuteAction::CreateLambda( [this] { bShowVertexNormals = !bShowVertexNormals; } ) );
+	RegisterCommonCommand( MeshEditorCommonCommands.ShowVertexNormals, FExecuteAction::CreateLambda( [this] { bShowVertexNormals = !bShowVertexNormals; } ) );
 
-	RegisterCommand( MeshEditorCommonActions.MarqueeSelectVertices, FExecuteAction::CreateLambda( [this] { PerformMarqueeSelect( EEditableMeshElementType::Vertex ); } ) );
-	RegisterCommand( MeshEditorCommonActions.MarqueeSelectEdges, FExecuteAction::CreateLambda( [this] { PerformMarqueeSelect( EEditableMeshElementType::Edge ); } ) );
-	RegisterCommand( MeshEditorCommonActions.MarqueeSelectPolygons, FExecuteAction::CreateLambda( [this] { PerformMarqueeSelect( EEditableMeshElementType::Polygon ); } ) );
-	RegisterCommand( MeshEditorCommonActions.FrameSelectedElements, FExecuteAction::CreateLambda( [this] { bShouldFocusToSelection = true; } ) );
+	RegisterCommonCommand( MeshEditorCommonCommands.MarqueeSelectVertices, FExecuteAction::CreateLambda( [this] { PerformMarqueeSelect( EEditableMeshElementType::Vertex ); } ) );
+	RegisterCommonCommand( MeshEditorCommonCommands.MarqueeSelectEdges, FExecuteAction::CreateLambda( [this] { PerformMarqueeSelect( EEditableMeshElementType::Edge ); } ) );
+	RegisterCommonCommand( MeshEditorCommonCommands.MarqueeSelectPolygons, FExecuteAction::CreateLambda( [this] { PerformMarqueeSelect( EEditableMeshElementType::Polygon ); } ) );
+	RegisterCommonCommand( MeshEditorCommonCommands.FrameSelectedElements, FExecuteAction::CreateLambda( [this] { bShouldFocusToSelection = true; } ) );
 
-	RegisterCommand( MeshEditorCommonActions.SetVertexSelectionMode, FExecuteAction::CreateLambda( [this] { SetMeshElementSelectionMode( EEditableMeshElementType::Vertex ); } ) );
-	RegisterCommand( MeshEditorCommonActions.SetEdgeSelectionMode, FExecuteAction::CreateLambda( [this] { SetMeshElementSelectionMode( EEditableMeshElementType::Edge ); } ) );
-	RegisterCommand( MeshEditorCommonActions.SetPolygonSelectionMode, FExecuteAction::CreateLambda( [this] { SetMeshElementSelectionMode( EEditableMeshElementType::Polygon ); } ) );
-	RegisterCommand( MeshEditorCommonActions.SetAnySelectionMode, FExecuteAction::CreateLambda( [this] { SetMeshElementSelectionMode( EEditableMeshElementType::Any ); } ) );
+	RegisterCommonCommand( MeshEditorCommonCommands.SetVertexSelectionMode, FExecuteAction::CreateLambda( [this] { SetMeshElementSelectionMode( EEditableMeshElementType::Vertex ); } ) );
+	RegisterCommonCommand( MeshEditorCommonCommands.SetEdgeSelectionMode, FExecuteAction::CreateLambda( [this] { SetMeshElementSelectionMode( EEditableMeshElementType::Edge ); } ) );
+	RegisterCommonCommand( MeshEditorCommonCommands.SetPolygonSelectionMode, FExecuteAction::CreateLambda( [this] { SetMeshElementSelectionMode( EEditableMeshElementType::Polygon ); } ) );
+	RegisterCommonCommand( MeshEditorCommonCommands.SetAnySelectionMode, FExecuteAction::CreateLambda( [this] { SetMeshElementSelectionMode( EEditableMeshElementType::Any ); } ) );
 
-	RegisterCommand( MeshEditorCommonActions.QuadrangulateMesh, FExecuteAction::CreateLambda( [ this ] { QuadrangulateMesh(); } ) );
+	RegisterCommonCommand( MeshEditorCommonCommands.QuadrangulateMesh, FExecuteAction::CreateLambda( [ this ] { QuadrangulateMesh(); } ) );
 
 	// Register element-specific commands
-	RegisterVertexCommand( MeshEditorVertexActions.WeldVertices, FExecuteAction::CreateLambda( [this] { WeldSelectedVertices(); } ) );
+	RegisterVertexCommand( MeshEditorVertexCommands.WeldVertices, FExecuteAction::CreateLambda( [this] { WeldSelectedVertices(); } ) );
 
-	RegisterEdgeCommand( MeshEditorEdgeActions.SelectEdgeLoop, FExecuteAction::CreateLambda( [ this ] { SelectEdgeLoops(); } ) );
+	RegisterEdgeCommand( MeshEditorEdgeCommands.SelectEdgeLoop, FExecuteAction::CreateLambda( [ this ] { SelectEdgeLoops(); } ) );
 
-	RegisterPolygonCommand( MeshEditorPolygonActions.FlipPolygon, FExecuteAction::CreateLambda( [this] { FlipSelectedPolygons(); } ) );
-	RegisterPolygonCommand( MeshEditorPolygonActions.TriangulatePolygon, FExecuteAction::CreateLambda( [this] { TriangulateSelectedPolygons(); } ) );
-	RegisterPolygonCommand( MeshEditorPolygonActions.AssignMaterial, FExecuteAction::CreateLambda( [this] { AssignSelectedMaterialToSelectedPolygons(); } ) );
+	RegisterPolygonCommand( MeshEditorPolygonCommands.FlipPolygon, FExecuteAction::CreateLambda( [this] { FlipSelectedPolygons(); } ) );
+	RegisterPolygonCommand( MeshEditorPolygonCommands.TriangulatePolygon, FExecuteAction::CreateLambda( [this] { TriangulateSelectedPolygons(); } ) );
+	RegisterPolygonCommand( MeshEditorPolygonCommands.AssignMaterial, FExecuteAction::CreateLambda( [this] { AssignSelectedMaterialToSelectedPolygons(); } ) );
 
 	for( TObjectIterator<UMeshEditorCommand> CommandCDOIter( RF_NoFlags ); CommandCDOIter; ++CommandCDOIter )
 	{
@@ -657,9 +668,17 @@ void FMeshEditorMode::RegisterPolygonEditingMode( const TSharedPtr<FUICommandInf
 }
 
 
-void FMeshEditorMode::RegisterCommand( const TSharedPtr<FUICommandInfo>& Command, const FExecuteAction& ExecuteAction )
+void FMeshEditorMode::RegisterCommonCommand( const TSharedPtr<FUICommandInfo>& Command, const FExecuteAction& ExecuteAction )
+{
+	CommonActions.Emplace( Command, FUIAction( ExecuteAction, FCanExecuteAction::CreateLambda( [this] { return GetSelectedEditableMeshes().Num() > 0; } ) ) );
+	VertexActions.Emplace( Command, FUIAction( ExecuteAction, FCanExecuteAction::CreateLambda( [this] { return GetSelectedEditableMeshes().Num() > 0; } ) ) );
+	EdgeActions.Emplace( Command, FUIAction( ExecuteAction, FCanExecuteAction::CreateLambda( [this] { return GetSelectedEditableMeshes().Num() > 0; } ) ) );
+	PolygonActions.Emplace( Command, FUIAction( ExecuteAction, FCanExecuteAction::CreateLambda( [this] { return GetSelectedEditableMeshes().Num() > 0; } ) ) );
+}
+
+
+void FMeshEditorMode::RegisterAnyElementCommand( const TSharedPtr<FUICommandInfo>& Command, const FExecuteAction& ExecuteAction )
 {
-	CommonActions.Emplace( Command, FUIAction( ExecuteAction, FCanExecuteAction::CreateLambda( [this] { return GetSelectedMeshElementType() != EEditableMeshElementType::Invalid;} ) ) );
 	VertexActions.Emplace( Command, FUIAction( ExecuteAction, FCanExecuteAction::CreateLambda( [this] { return GetSelectedMeshElementType() != EEditableMeshElementType::Invalid;} ) ) );
 	EdgeActions.Emplace( Command, FUIAction( ExecuteAction, FCanExecuteAction::CreateLambda( [this] { return GetSelectedMeshElementType() != EEditableMeshElementType::Invalid;} ) ) );
 	PolygonActions.Emplace( Command, FUIAction( ExecuteAction, FCanExecuteAction::CreateLambda( [this] { return GetSelectedMeshElementType() != EEditableMeshElementType::Invalid;} ) ) );
@@ -726,6 +745,8 @@ void FMeshEditorMode::Enter()
 		Toolkit = MakeShared< FMeshEditorModeToolkit >( *UIContract );
 		Toolkit->Init( Owner->GetToolkitHost() );
 	}
+
+	UpdateSelectedEditableMeshes();
 }
 
 
@@ -1006,11 +1027,17 @@ bool FMeshEditorMode::InputKey( FEditorViewportClient* ViewportClient, FViewport
 
 	if( Event == IE_Pressed )
 	{
-		const FUICommandList* CommandList = GetCommandListForSelectedElementType();
-		check( CommandList );
-
 		FModifierKeysState ModifierKeysState = FSlateApplication::Get().GetModifierKeys();
-		bHandled = CommandList->ProcessCommandBindings( Key, ModifierKeysState, false );
+
+		const FUICommandList* CommandList = GetCommandListForSelectedElementType();
+		if( CommandList != nullptr )
+		{
+			bHandled = CommandList->ProcessCommandBindings( Key, ModifierKeysState, false );
+		}
+		else
+		{
+			bHandled = CommonCommands->ProcessCommandBindings( Key, ModifierKeysState, false );
+		}
 	}
 
 	else if( Event == IE_DoubleClick )
@@ -1036,9 +1063,11 @@ const FUICommandList* FMeshEditorMode::GetCommandListForSelectedElementType() co
 		case EEditableMeshElementType::Polygon:
 			return PolygonCommands.Get();
 
-		default:
-			return CommonCommands.Get();
+		case EEditableMeshElementType::Any:
+			return AnyElementCommands.Get();
 	}
+
+	return nullptr;
 }
 
 
@@ -1120,35 +1149,18 @@ void FMeshEditorMode::CommitEditableMeshIfNecessary( UEditableMesh* EditableMesh
 
 void FMeshEditorMode::CommitSelectedMeshes()
 {
-	// @todo mesheditor: this is inefficient as it iterates the mesh elements arrays each time round the loop to fix up FMeshElements
-	for( int32 SelectedElementIndex = 0; SelectedElementIndex < this->SelectedMeshElements.Num(); ++SelectedElementIndex )
+	for( FComponentAndEditableMesh& ComponentAndEditableMesh : SelectedComponentsAndEditableMeshes )
 	{
-		const FMeshElement& SelectedMeshElement = this->SelectedMeshElements[ SelectedElementIndex ];
-		if( SelectedMeshElement.IsValidMeshElement() )
-		{
-			UEditableMesh* EditableMesh = this->FindOrCreateEditableMesh( *SelectedMeshElement.Component, SelectedMeshElement.ElementAddress.SubMeshAddress );
-			if( EditableMesh != nullptr )
-			{
-				CommitEditableMeshIfNecessary( EditableMesh, SelectedMeshElement.Component.Get() );
-			}
-		}
+		CommitEditableMeshIfNecessary( ComponentAndEditableMesh.EditableMesh, ComponentAndEditableMesh.Component );
 	}
 }
 
 
 void FMeshEditorMode::PropagateInstanceChanges()
 {
-	for( int32 SelectedElementIndex = 0; SelectedElementIndex < SelectedMeshElements.Num(); ++SelectedElementIndex )
+	for( UEditableMesh* EditableMesh : SelectedEditableMeshes )
 	{
-		const FMeshElement& SelectedMeshElement = SelectedMeshElements[ SelectedElementIndex ];
-		if( SelectedMeshElement.IsValidMeshElement() )
-		{
-			UEditableMesh* EditableMesh = FindOrCreateEditableMesh( *SelectedMeshElement.Component, SelectedMeshElement.ElementAddress.SubMeshAddress );
-			if( EditableMesh != nullptr )
-			{
-				EditableMesh->PropagateInstanceChanges();
-			}
-		}
+		EditableMesh->PropagateInstanceChanges();
 	}
 
 	CachedEditableMeshes.Empty();
@@ -1160,16 +1172,12 @@ bool FMeshEditorMode::CanPropagateInstanceChanges() const
 	// @todo mesheditor: this could be more thorough:
 	// it should not allow instance changes to be propagated if more than one instance is selected which derives from the same static mesh.
 	// However MeshEditorMode has no generic way to know if this is the case (and it's unclear how that might be presented in the API).
-	for( int32 SelectedElementIndex = 0; SelectedElementIndex < SelectedMeshElements.Num(); ++SelectedElementIndex )
+	const TArray<UEditableMesh*>& SelectedEditableMeshes = GetSelectedEditableMeshes();
+	for( const UEditableMesh* EditableMesh : SelectedEditableMeshes )
 	{
-		const FMeshElement& SelectedMeshElement = SelectedMeshElements[ SelectedElementIndex ];
-		if( SelectedMeshElement.IsValidMeshElement() )
+		if( EditableMesh != nullptr && EditableMesh->IsCommittedAsInstance() )
 		{
-			UEditableMesh* EditableMesh = CachedEditableMeshes.FindRef( SelectedMeshElement.ElementAddress.SubMeshAddress );
-			if( EditableMesh != nullptr && EditableMesh->IsCommittedAsInstance() )
-			{
-				return true;
-			}
+			return true;
 		}
 	}
 
@@ -1310,10 +1318,7 @@ void FMeshEditorMode::AddOrRemoveSubdivisionLevel( const bool bShouldAdd )
 {
 	if( ActiveAction == NAME_None )
 	{
-		static TSet< UEditableMesh* > SelectedMeshes;
-		SelectedMeshes.Reset();
-
-		if( SelectedMeshElements.Num() == 0 )
+		if( GetSelectedEditableMeshes().Num() == 0 )
 		{
 			return;
 		}
@@ -1324,18 +1329,7 @@ void FMeshEditorMode::AddOrRemoveSubdivisionLevel( const bool bShouldAdd )
 
 		CommitSelectedMeshes();
 
-		for( int32 SelectedElementIndex = 0; SelectedElementIndex < SelectedMeshElements.Num(); ++SelectedElementIndex )
-		{
-			const FMeshElement& SelectedMeshElement = SelectedMeshElements[ SelectedElementIndex ];
-			if( SelectedMeshElement.IsValidMeshElement() )
-			{
-				UEditableMesh* EditableMesh = FindOrCreateEditableMesh( *SelectedMeshElement.Component, SelectedMeshElement.ElementAddress.SubMeshAddress );
-				if( EditableMesh != nullptr )
-				{
-					SelectedMeshes.Add( EditableMesh );
-				}
-			}
-		}
+		const TArray<UEditableMesh*>& SelectedMeshes = GetSelectedEditableMeshes();
 
 		for( UEditableMesh* EditableMesh : SelectedMeshes )
 		{
@@ -1363,10 +1357,7 @@ void FMeshEditorMode::QuadrangulateMesh()
 {
 	if( ActiveAction == NAME_None )
 	{
-		static TSet< UEditableMesh* > SelectedMeshes;
-		SelectedMeshes.Reset();
-
-		if( SelectedMeshElements.Num() == 0 )
+		if( GetSelectedEditableMeshes().Num() == 0 )
 		{
 			return;
 		}
@@ -1375,18 +1366,7 @@ void FMeshEditorMode::QuadrangulateMesh()
 
 		CommitSelectedMeshes();
 
-		for( int32 SelectedElementIndex = 0; SelectedElementIndex < SelectedMeshElements.Num(); ++SelectedElementIndex )
-		{
-			const FMeshElement& SelectedMeshElement = SelectedMeshElements[ SelectedElementIndex ];
-			if( SelectedMeshElement.IsValidMeshElement() )
-			{
-				UEditableMesh* EditableMesh = FindOrCreateEditableMesh( *SelectedMeshElement.Component, SelectedMeshElement.ElementAddress.SubMeshAddress );
-				if( EditableMesh != nullptr )
-				{
-					SelectedMeshes.Add( EditableMesh );
-				}
-			}
-		}
+		const TArray<UEditableMesh*>& SelectedMeshes = GetSelectedEditableMeshes();
 
 		DeselectAllMeshElements();
 
@@ -1900,6 +1880,13 @@ void FMeshEditorMode::AddReferencedObjects( FReferenceCollector& Collector )
 	Collector.AddReferencedObject( HoveredGeometryMaterial );
 	Collector.AddReferencedObject( HoveredFaceMaterial );
 
+	Collector.AddReferencedObjects( SelectedEditableMeshes );
+	for( FComponentAndEditableMesh& ComponentAndEditableMesh : SelectedComponentsAndEditableMeshes )
+	{
+		Collector.AddReferencedObject( ComponentAndEditableMesh.Component );
+		Collector.AddReferencedObject( ComponentAndEditableMesh.EditableMesh );
+	}
+
 	for( auto Pair : CachedEditableMeshes )
 	{
 		Collector.AddReferencedObject( Pair.Value );
@@ -2224,38 +2211,10 @@ void FMeshEditorMode::Render( const FSceneView* SceneView, FViewport* Viewport,
 
 	// Draw all polygon edges for selected/hovered meshes (@todo mesheditor: This only draws polygons right now.  We should also draw "floating" edges and vertices too!)
 	{
-		struct FComponentAndEditableMesh
-		{
-			const UPrimitiveComponent* Component;
-			const UEditableMesh* EditableMesh;
-
-			FComponentAndEditableMesh( const UPrimitiveComponent* InitComponent, const UEditableMesh* InitEditableMesh )
-				: Component( InitComponent ),
-				  EditableMesh( InitEditableMesh )
-			{
-			}
-
-			inline bool operator==( const FComponentAndEditableMesh& Other ) const
-			{
-				return Component == Other.Component && EditableMesh == Other.EditableMesh;
-			}
-		};
-
-
 		static TArray<FComponentAndEditableMesh> HoveredOrSelectedMeshes;
 		HoveredOrSelectedMeshes.Reset();
 
-		for( const FMeshElement& SelectedMeshElement : SelectedMeshElements )	// @todo mesheditor: Should also include meshes that are editor-selected
-		{
-			if( SelectedMeshElement.IsValidMeshElement() )
-			{
-				UEditableMesh* EditableMesh = FindOrCreateEditableMesh( *SelectedMeshElement.Component, SelectedMeshElement.ElementAddress.SubMeshAddress );
-				if( EditableMesh != nullptr )
-				{
-					HoveredOrSelectedMeshes.AddUnique( FComponentAndEditableMesh( SelectedMeshElement.Component.Get(), EditableMesh ) );
-				}
-			}
-		}
+		HoveredOrSelectedMeshes = SelectedComponentsAndEditableMeshes;
 
 		// Only draw hover if we're not in the middle of an interactive edit
 		if( ActiveAction == NAME_None )
@@ -2655,8 +2614,8 @@ void FMeshEditorMode::OnViewportInteractionHoverUpdate( UViewportInteractor* Vie
 								// they're in subdivision preview mode.  Their base cage mesh won't match their complex collision geometry, but we
 								// still need to allow the user to interact with mesh elements outside the bounds of that geometry.
 								FEditableMeshSubMeshAddress SubMeshAddress = UEditableMeshFactory::MakeSubmeshAddress( Component, LODIndex );
-								UEditableMesh** EditableMeshPtr = CachedEditableMeshes.Find( SubMeshAddress );
-								if( bTraceComplex || ( EditableMeshPtr != nullptr && ( *EditableMeshPtr )->IsPreviewingSubdivisions() ) )
+								const UEditableMesh* EditableMesh = FindEditableMesh( *Component, SubMeshAddress );
+								if( bTraceComplex || ( EditableMesh != nullptr && EditableMesh->IsPreviewingSubdivisions() ) )
 								{
 									// Don't bother with editor-only 'helper' actors, we never want to visualize or edit geometry on those
 									if( !Component->IsEditorOnly() &&
@@ -2682,8 +2641,7 @@ void FMeshEditorMode::OnViewportInteractionHoverUpdate( UViewportInteractor* Vie
 
 						for( UPrimitiveComponent* HitComponent : HitComponents )
 						{
-							const bool bOnlyInteractWithSelectedObjects = GetDefault<UMeshEditorSettings>()->bOnlyEditSelectedObjects;
-							if( !bOnlyInteractWithSelectedObjects || GEditor->GetSelectedActors()->IsSelected( HitComponent->GetOwner() ) )
+							if( GEditor->GetSelectedActors()->IsSelected( HitComponent->GetOwner() ) )
 							{
 								// @todo mesheditor debug
 								// GHackComponentToWorld = HitComponent->GetComponentToWorld();
@@ -2965,7 +2923,7 @@ void FMeshEditorMode::UpdateActiveAction( const bool bIsActionFinishing )
 			UPrimitiveComponent* Component = nullptr;
 			FEditableMeshSubMeshAddress SubMeshAddress;
 
-			if( SelectedMeshElements.Num() == 0 )
+			if( SelectedEditableMeshes.Num() == 0 )
 			{
 				// @todo mesheditor: support creating a new mesh from scratch here
 				// Look into support for creating new assets in the transient package without needing to specify a filename?
@@ -2973,23 +2931,10 @@ void FMeshEditorMode::UpdateActiveAction( const bool bIsActionFinishing )
 			else
 			{
 				// Currently adds new vertices to whichever editable mesh is currently selected
+				Component = SelectedComponentsAndEditableMeshes[ 0 ].Component;
+				EditableMesh = SelectedComponentsAndEditableMeshes[ 1 ].EditableMesh;
 				// @todo mesheditor: allow multiple selected meshes? What should this do?
 
-				for( int32 SelectedElementIndex = 0; SelectedElementIndex < SelectedMeshElements.Num(); ++SelectedElementIndex )
-				{
-					const FMeshElement& SelectedMeshElement = SelectedMeshElements[ SelectedElementIndex ];
-					if( SelectedMeshElement.IsValidMeshElement() )
-					{
-						EditableMesh = FindOrCreateEditableMesh( *SelectedMeshElement.Component, SelectedMeshElement.ElementAddress.SubMeshAddress );
-						if( EditableMesh != nullptr )
-						{
-							Component = SelectedMeshElement.Component.Get();
-							SubMeshAddress = SelectedMeshElement.ElementAddress.SubMeshAddress;
-							break;
-						}
-					}
-				}
-
 				check( EditableMesh );
 				check( Component );
 
@@ -4233,7 +4178,6 @@ bool FMeshEditorMode::FrustumSelect( const FConvexVolume& InFrustum, FEditorView
 {
 	// @todo mesheditor urgent: settings class for bundling together all these kind of things
 	const bool bShouldDeselectAllFirst = true;	// @todo mesheditor: needs to be passed to this method
-	const bool bOnlyInteractWithSelectedObjects = GetDefault<UMeshEditorSettings>()->bOnlyEditSelectedObjects;
 	const bool bOnlySelectVisibleMeshes = GetDefault<UMeshEditorSettings>()->bOnlySelectVisibleMeshes;
 	const bool bOnlySelectVisibleElements = GetDefault<UMeshEditorSettings>()->bOnlySelectVisibleElements;
 
@@ -4320,7 +4264,7 @@ bool FMeshEditorMode::FrustumSelect( const FConvexVolume& InFrustum, FEditorView
 
 		for( AActor* Actor : HitActors )
 		{
-			if( !bOnlyInteractWithSelectedObjects || GEditor->GetSelectedActors()->IsSelected( Actor ) )
+			if( GEditor->GetSelectedActors()->IsSelected( Actor ) )
 			{
 				AddEditableMeshFromActor( Actor );
 			}
@@ -4336,7 +4280,7 @@ bool FMeshEditorMode::FrustumSelect( const FConvexVolume& InFrustum, FEditorView
 
 			if( !Actor->IsA( ABrush::StaticClass() ) &&
 			    !Actor->IsHiddenEd() &&
-			    ( !bOnlyInteractWithSelectedObjects || GEditor->GetSelectedActors()->IsSelected( Actor ) ) )
+			    GEditor->GetSelectedActors()->IsSelected( Actor ) )
 			{
 				AddEditableMeshFromActor( Actor );
 			}
@@ -4659,12 +4603,59 @@ void FMeshEditorMode::RefreshTransformables( const bool bNewObjectsSelected )
 }
 
 
+void FMeshEditorMode::UpdateSelectedEditableMeshes()
+{
+	SelectedEditableMeshes.Reset();
+	SelectedComponentsAndEditableMeshes.Reset();
+
+	// If we have selected elements, make sure those are in our set
+	for( FMeshElement& SelectedMeshElement : SelectedMeshElements )
+	{
+		if( SelectedMeshElement.IsValidMeshElement() )
+		{
+			UEditableMesh* EditableMesh = FindOrCreateEditableMesh( *SelectedMeshElement.Component, SelectedMeshElement.ElementAddress.SubMeshAddress );
+			if( EditableMesh != nullptr )
+			{
+				SelectedComponentsAndEditableMeshes.AddUnique( FComponentAndEditableMesh( SelectedMeshElement.Component.Get(), EditableMesh ) );
+				SelectedEditableMeshes.AddUnique( EditableMesh );
+			}
+		}
+	}
+
+	// Check the actors that are selected, and add any meshes we find
+	for( TSelectionIterator<FGenericSelectionFilter> SelectionIt( *GEditor->GetSelectedActors() ); SelectionIt; ++SelectionIt )
+	{
+		AActor* Actor = Cast<AActor>( *SelectionIt );
+		if( Actor != nullptr )
+		{
+			TArray<UActorComponent*> PrimitiveComponents = Actor->GetComponentsByClass( UPrimitiveComponent::StaticClass() );
+			for( UActorComponent* PrimitiveActorComponent : PrimitiveComponents )
+			{
+				UPrimitiveComponent* Component = CastChecked<UPrimitiveComponent>( PrimitiveActorComponent );
+
+				// Don't bother with editor-only 'helper' actors, we never want to visualize or edit geometry on those
+				if( !Component->IsEditorOnly() &&
+					( Component->GetOwner() == nullptr || !Component->GetOwner()->IsEditorOnly() ) )
+				{
+					const int32 LODIndex = 0;			// @todo mesheditor: We'll want to select an LOD to edit in various different wants (LOD that's visible, or manual user select, etc.)
+
+					FEditableMeshSubMeshAddress SubMeshAddress = UEditableMeshFactory::MakeSubmeshAddress( Component, LODIndex );
+					UEditableMesh* EditableMesh = FindOrCreateEditableMesh( *Component, SubMeshAddress );
+					if( EditableMesh != nullptr )
+					{
+						SelectedComponentsAndEditableMeshes.AddUnique( FComponentAndEditableMesh( Component, EditableMesh ) );
+						SelectedEditableMeshes.AddUnique( EditableMesh );
+					}
+				}
+			}
+		}
+	}
+}
+
+
 void FMeshEditorMode::OnActorSelectionChanged( const TArray<UObject*>& NewSelection, bool bForceRefresh )
 {
-	// If we're configured to only interact with selected objects, deselect any elements that no longer belong
-	// to the selected set of actors.
-	const bool bOnlyInteractWithSelectedObjects = GetDefault<UMeshEditorSettings>()->bOnlyEditSelectedObjects;
-	if( bOnlyInteractWithSelectedObjects )
+	// Deselect any elements that no longer belong to the selected set of actors.
 	{
 		// Don't respond to actor selection changes if a transaction isn't in progress, because it's probably
 		// initiated from an undo/redo action itself, in which case the selection state changes will already
@@ -4688,6 +4679,9 @@ void FMeshEditorMode::OnActorSelectionChanged( const TArray<UObject*>& NewSelect
 			}
 		}
 	}
+
+	// Update our set of selected meshes
+	UpdateSelectedEditableMeshes();
 }
 
 void FMeshEditorMode::MakeVRRadialMenuActionsMenu(FMenuBuilder& MenuBuilder, TSharedPtr<FUICommandList> CommandList, class UVREditorMode* VRMode, float& RadiusOverride)
diff --git a/Engine/Plugins/Editor/MeshEditor/Source/MeshEditor/MeshEditorMode.h b/Engine/Plugins/Editor/MeshEditor/Source/MeshEditor/MeshEditorMode.h
--- a/Engine/Plugins/Editor/MeshEditor/Source/MeshEditor/MeshEditorMode.h
+++ b/Engine/Plugins/Editor/MeshEditor/Source/MeshEditor/MeshEditorMode.h
@@ -156,6 +156,9 @@ protected:
 	virtual EEditableMeshElementType GetSelectedMeshElementType() const override;
 	virtual bool IsMeshElementTypeSelected( EEditableMeshElementType ElementType ) const override { return GetSelectedMeshElementType() == ElementType; }
 	virtual bool IsMeshElementTypeSelectedOrIsActiveSelectionMode( EEditableMeshElementType ElementType ) const override { return GetMeshElementSelectionMode() == ElementType || GetSelectedMeshElementType() == ElementType; }
+	virtual const TArray<UEditableMesh*>& GetSelectedEditableMeshes() const override { return SelectedEditableMeshes; }
+	virtual const TArray<UEditableMesh*>& GetSelectedEditableMeshes() override { return SelectedEditableMeshes; }
+	virtual const TArray<TTuple<TSharedPtr<FUICommandInfo>, FUIAction>>& GetCommonActions() const override { return CommonActions; }
 	virtual const TArray<TTuple<TSharedPtr<FUICommandInfo>, FUIAction>>& GetVertexActions() const override { return VertexActions; }
 	virtual const TArray<TTuple<TSharedPtr<FUICommandInfo>, FUIAction>>& GetEdgeActions() const override { return EdgeActions; }
 	virtual const TArray<TTuple<TSharedPtr<FUICommandInfo>, FUIAction>>& GetPolygonActions() const override { return PolygonActions; }
@@ -251,6 +254,9 @@ protected:
 	/** Returns the index of an element in the selection set, or INDEX_NONE if its not selected */
 	int32 GetSelectedMeshElementIndex( const FMeshElement& MeshElement ) const;
 
+	/** Updates our list of selected editable meshes based on which actors/components are selected.  Will create new editable meshes as needed */
+	void UpdateSelectedEditableMeshes();
+
 	/** Updates the current view location, from either the viewport client or the VR interface, whichever is in use */
 	void UpdateCameraToWorldTransform( const FEditorViewportClient& ViewportClient );
 
@@ -269,7 +275,8 @@ protected:
 	void RegisterEdgeEditingMode( const TSharedPtr<FUICommandInfo>& Command, FName EditingMode );
 	void RegisterPolygonEditingMode( const TSharedPtr<FUICommandInfo>& Command, FName EditingMode );
 
-	void RegisterCommand( const TSharedPtr<FUICommandInfo>& Command, const FExecuteAction& ExecuteAction );
+	void RegisterCommonCommand( const TSharedPtr<FUICommandInfo>& Command, const FExecuteAction& ExecuteAction );
+	void RegisterAnyElementCommand( const TSharedPtr<FUICommandInfo>& Command, const FExecuteAction& ExecuteAction );
 	void RegisterVertexCommand( const TSharedPtr<FUICommandInfo>& Command, const FExecuteAction& ExecuteAction );
 	void RegisterEdgeCommand( const TSharedPtr<FUICommandInfo>& Command, const FExecuteAction& ExecuteAction );
 	void RegisterPolygonCommand( const TSharedPtr<FUICommandInfo>& Command, const FExecuteAction& ExecuteAction );
@@ -492,6 +499,31 @@ protected:
 	/** Specifies the type of element which is currently being selected */
 	EEditableMeshElementType MeshElementSelectionMode;
 
+	/** The editable meshes that are currently selected.  This generally maps to selected actors/components, and
+	    may contain meshes even when no elements are selected. */
+	TArray<UEditableMesh*> SelectedEditableMeshes;
+
+	struct FComponentAndEditableMesh
+	{
+		UPrimitiveComponent* Component;
+		UEditableMesh* EditableMesh;
+
+		FComponentAndEditableMesh( UPrimitiveComponent* InitComponent, UEditableMesh* InitEditableMesh )
+			: Component( InitComponent ),
+			  EditableMesh( InitEditableMesh )
+		{
+		}
+
+		inline bool operator==( const FComponentAndEditableMesh& Other ) const
+		{
+			return Component == Other.Component && EditableMesh == Other.EditableMesh;
+		}
+	};
+
+	/** The individual components that are selected along with the editable mesh for each component.  Note that the same
+	    editable mesh may appear more than once in this list, because it could be shared between components! */
+	TArray<FComponentAndEditableMesh> SelectedComponentsAndEditableMeshes;
+
 	/** List of mesh elements that we've selected.  All elements in this list will always have the same mesh element type.
 	    We don't allow users to select both edges, faces and/or polygons at the same time. */
 	TArray<FMeshElement> SelectedMeshElements;
@@ -556,6 +588,9 @@ protected:
 	/** Command list for actions available regardless of selection */
 	TSharedPtr<FUICommandList> CommonCommands;
 
+	/** Command list for actions available regardless of selection, as long as something is selected */
+	TSharedPtr<FUICommandList> AnyElementCommands;
+
 	/** Command list for actions available when a vertex is selected */
 	TSharedPtr<FUICommandList> VertexCommands;
 
diff --git a/Engine/Plugins/Editor/MeshEditor/Source/MeshEditor/MeshEditorModeToolkit.cpp b/Engine/Plugins/Editor/MeshEditor/Source/MeshEditor/MeshEditorModeToolkit.cpp
--- a/Engine/Plugins/Editor/MeshEditor/Source/MeshEditor/MeshEditorModeToolkit.cpp
+++ b/Engine/Plugins/Editor/MeshEditor/Source/MeshEditor/MeshEditorModeToolkit.cpp
@@ -183,12 +183,25 @@ void SMeshEditorModeControls::Construct( const FArguments& InArgs, IMeshEditorMo
 
 	WidgetSwitcher->AddSlot( static_cast<int32>( EEditableMeshElementType::Invalid ) )
 	[
-		SNew( SBox )
-		.Padding( 20.0f )
-		.HAlign( HAlign_Center )
+		SNew( SVerticalBox )
+		+ SVerticalBox::Slot()
 		[
-			SNew( STextBlock )
-			.Text( LOCTEXT( "NothingSelected", "Nothing is selected" ) )
+			SNew( SBox )
+			.Visibility_Lambda( [&MeshEditorMode]() { return MeshEditorMode.GetSelectedEditableMeshes().Num() > 0 ? EVisibility::Visible : EVisibility::Collapsed; } )
+			[
+				SNew( SMeshEditorModeControlWidget, MeshEditorMode.GetCommonActions() )
+			]
+		]
+		+SVerticalBox::Slot()
+		[
+			SNew( SBox )
+			.Padding( 20.0f )
+			.HAlign( HAlign_Center )
+			.Visibility_Lambda( [&MeshEditorMode]() { return MeshEditorMode.GetSelectedEditableMeshes().Num() == 0 ? EVisibility::Visible : EVisibility::Collapsed; } )
+			[
+				SNew( STextBlock )
+				.Text( LOCTEXT( "NothingSelected", "Nothing is selected" ) )
+			]
 		]
 	];
 
diff --git a/Engine/Plugins/Editor/MeshEditor/Source/MeshEditor/MeshEditorSettings.h b/Engine/Plugins/Editor/MeshEditor/Source/MeshEditor/MeshEditorSettings.h
--- a/Engine/Plugins/Editor/MeshEditor/Source/MeshEditor/MeshEditorSettings.h
+++ b/Engine/Plugins/Editor/MeshEditor/Source/MeshEditor/MeshEditorSettings.h
@@ -18,7 +18,6 @@ public:
 
 	UMeshEditorSettings()
 		: bSeparateSelectionSetPerMode( false )
-		, bOnlyEditSelectedObjects( true )
 		, bOnlySelectVisibleMeshes( true )
 		, bOnlySelectVisibleElements( true )
 		, bAllowGrabberSphere( false )
@@ -31,10 +30,6 @@ public:
 	UPROPERTY(EditAnywhere, config, Category=MeshEditor)
 	bool bSeparateSelectionSetPerMode;
 
-	/** Whether you must select objects before mesh editing features will be available.  This helps to prevent accidental mis-clicks. */
-	UPROPERTY(EditAnywhere, config, Category=MeshEditor)
-	bool bOnlyEditSelectedObjects;
-
 	/** Whether only unoccluded meshes will be selected by marquee select, or whether all meshes within the selection box will be selected, regardless of whether they are behind another. */
 	UPROPERTY(EditAnywhere, config, Category=MeshEditor)
 	bool bOnlySelectVisibleMeshes;
diff --git a/Engine/Plugins/Editor/MeshEditor/Source/MeshEditor/Public/IMeshEditorModeUIContract.h b/Engine/Plugins/Editor/MeshEditor/Source/MeshEditor/Public/IMeshEditorModeUIContract.h
--- a/Engine/Plugins/Editor/MeshEditor/Source/MeshEditor/Public/IMeshEditorModeUIContract.h
+++ b/Engine/Plugins/Editor/MeshEditor/Source/MeshEditor/Public/IMeshEditorModeUIContract.h
@@ -25,6 +25,13 @@ public:
 	/** Returns whether either the specified element type is selected, or we're in the selection mode for that element type */
 	virtual bool IsMeshElementTypeSelectedOrIsActiveSelectionMode( EEditableMeshElementType ElementType ) const = 0;
 
+	/** Gets all of the editable meshes that are currently selected (const) */
+	virtual const TArray<class UEditableMesh*>& GetSelectedEditableMeshes() const = 0;
+
+	/** Gets all of the editable meshes that are currently selected */
+	virtual const TArray<class UEditableMesh*>& GetSelectedEditableMeshes() = 0;
+
+	virtual const TArray<TTuple<TSharedPtr<FUICommandInfo>, FUIAction>>& GetCommonActions() const = 0;
 	virtual const TArray<TTuple<TSharedPtr<FUICommandInfo>, FUIAction>>& GetVertexActions() const = 0;
 	virtual const TArray<TTuple<TSharedPtr<FUICommandInfo>, FUIAction>>& GetEdgeActions() const = 0;
 	virtual const TArray<TTuple<TSharedPtr<FUICommandInfo>, FUIAction>>& GetPolygonActions() const = 0;
diff --git a/Engine/Plugins/Editor/MeshEditor/Source/MeshEditor/Public/MeshEditorCommands.h b/Engine/Plugins/Editor/MeshEditor/Source/MeshEditor/Public/MeshEditorCommands.h
--- a/Engine/Plugins/Editor/MeshEditor/Source/MeshEditor/Public/MeshEditorCommands.h
+++ b/Engine/Plugins/Editor/MeshEditor/Source/MeshEditor/Public/MeshEditorCommands.h
@@ -156,7 +156,7 @@ public:
 
 
 
-// Actions that can be invoked from this mode regardless of what type of elements are selected
+// Actions that can be invoked from this mode as long as at least one mesh is selected
 class FMeshEditorCommonCommands : public TCommands<FMeshEditorCommonCommands>
 {
 public:
@@ -168,9 +168,6 @@ public:
 
 public:
 
-	/** Deletes selected mesh elements, including polygons partly defined by selected elements */
-	TSharedPtr<FUICommandInfo> DeleteMeshElement;
-
 	/** Increases the number of subdivision levels for the selected mesh. */
 	TSharedPtr<FUICommandInfo> AddSubdivisionLevel;
 
@@ -202,6 +199,23 @@ public:
 };
 
 
+// Actions that can be invoked from this mode for any type of selected element
+class FMeshEditorAnyElementCommands : public TCommands<FMeshEditorAnyElementCommands>
+{
+public:
+	FMeshEditorAnyElementCommands();
+
+	// TCommands<> interface
+	virtual void RegisterCommands() override;
+	// End of TCommands<> interface
+
+public:
+
+	/** Deletes selected mesh elements, including polygons partly defined by selected elements */
+	TSharedPtr<FUICommandInfo> DeleteMeshElement;
+};
+
+
 // Actions that can be invoked from this mode when vertices are selected
 class FMeshEditorVertexCommands : public TCommands<FMeshEditorVertexCommands>
 {
