Mesh Editor: Frame Selected will now frame the actor/component if no mesh elements are selected

diff --git a/Engine/Plugins/Editor/MeshEditor/Source/MeshEditor/MeshEditorMode.cpp b/Engine/Plugins/Editor/MeshEditor/Source/MeshEditor/MeshEditorMode.cpp
--- a/Engine/Plugins/Editor/MeshEditor/Source/MeshEditor/MeshEditorMode.cpp
+++ b/Engine/Plugins/Editor/MeshEditor/Source/MeshEditor/MeshEditorMode.cpp
@@ -906,7 +906,19 @@ void FMeshEditorMode::Tick( FEditorViewportClient* ViewportClient, float DeltaTi
 	if( bShouldFocusToSelection )
 	{
 		bShouldFocusToSelection = false;
-		FrameSelectedElements( ViewportClient );
+
+		// Are any elements selected?  If so, we'll focus directly on those
+		if( SelectedMeshElements.Num() > 0 )
+		{
+			FrameSelectedElements( ViewportClient );
+		}
+		else
+		{
+			// No elements selected, so focus on selected actors/components instead.
+			TArray<UObject*> SelectedActors;
+			GEditor->GetSelectedActors()->GetSelectedObjects( AActor::StaticClass(), /* Out */ SelectedActors );
+			GEditor->MoveViewportCamerasToActor( *reinterpret_cast<TArray<AActor*>*>( &SelectedActors ), true );
+		}
 	}
 
 	// @todo mesheditor: Should take into account world scaling while in VR (room space interactor movement threshold)
