Mesh Editor: Lots of fixes for orthographic view supprt
- Added proper rendered wire sizing for orthographic views
- Fixed various backface checks
- Fixed fuzzy selection checks in orthographic views
- Fixed elements being selected from opposite direction when clicking

diff --git a/Engine/Plugins/Editor/MeshEditor/Source/MeshEditor/MeshEditorMode.cpp b/Engine/Plugins/Editor/MeshEditor/Source/MeshEditor/MeshEditorMode.cpp
--- a/Engine/Plugins/Editor/MeshEditor/Source/MeshEditor/MeshEditorMode.cpp
+++ b/Engine/Plugins/Editor/MeshEditor/Source/MeshEditor/MeshEditorMode.cpp
@@ -69,7 +69,8 @@ namespace MeshEd
 	static FAutoConsoleVariable OverlayVertexSize( TEXT( "MeshEd.OverlayVertexSize" ), 4.0f, TEXT( "How large a vertex is on a mesh overlay" ) );
 	static FAutoConsoleVariable OverlayLineThickness( TEXT( "MeshEd.OverlayLineThickness" ), 0.9f, TEXT( "How thick overlay lines should be on top of meshes when hovered or selected" ) );
 	static FAutoConsoleVariable OverlayDistanceScaleFactor( TEXT( "MeshEd.OverlayDistanceScaleFactor" ), 0.002f, TEXT( "How much to scale overlay wires automatically based on distance to the viewer" ) );
-	static FAutoConsoleVariable OverlayDistanceBias( TEXT( "MeshEd.OverlayDistanceBias" ), 0.05f, TEXT( "How much to bias distance scale by, regardless of distance to the viewer" ) );
+	static FAutoConsoleVariable OverlayPerspectiveDistanceBias( TEXT( "MeshEd.OverlayPerspectiveDistanceBias" ), 0.05f, TEXT( "How much to bias distance scale by in perspective views, regardless of distance to the viewer" ) );
+	static FAutoConsoleVariable OverlayOrthographicDistanceBias( TEXT( "MeshEd.OverlayOrthographicDistanceBias" ), 1.0f, TEXT( "How much to bias distance scale by in orthograph views, regardless of distance to the viewer" ) );
 	static FAutoConsoleVariable SelectedSizeBias( TEXT( "MeshEd.SelectedSizeBias" ), 0.1f, TEXT( "Selected mesh element size bias" ) );
 	static FAutoConsoleVariable SelectedAnimationExtraSizeBias( TEXT( "MeshEd.SelectedAnimationExtraSizeBias" ), 2.5f, TEXT( "Extra hovered mesh element size bias when animating" ) );
 	static FAutoConsoleVariable HoveredSizeBias( TEXT( "MeshEd.HoveredSizeBias" ), 0.1f, TEXT( "Selected mesh element size bias" ) );
@@ -1018,6 +1019,7 @@ void FMeshEditorMode::UpdateCameraToWorldTransform( const FEditorViewportClient&
 	{
 		CachedCameraToWorld = FTransform( ViewportClient.GetViewTransform().GetRotation(), ViewportClient.GetViewTransform().GetLocation() );
 	}
+	bCachedIsPerspectiveView = ViewportClient.IsPerspective();
 }
 
 
@@ -1917,30 +1919,30 @@ void FMeshEditorMode::AddReferencedObjects( FReferenceCollector& Collector )
 }
 
 
-void FMeshEditorMode::AddVertexToDynamicMesh( const UEditableMesh& EditableMesh, const FTransform& CameraToWorld, const FMatrix& ComponentToWorldMatrix, const FVertexID VertexID, const FColor ColorAndOpacity, const float SizeBias, const bool bApplyDepthBias, FDynamicMeshBuilder& MeshBuilder )
+void FMeshEditorMode::AddVertexToDynamicMesh( const UEditableMesh& EditableMesh, const FTransform& CameraToWorld, const bool bIsPerspectiveView, const FMatrix& ComponentToWorldMatrix, const FVertexID VertexID, const FColor ColorAndOpacity, const float SizeBias, const bool bApplyDepthBias, FDynamicMeshBuilder& MeshBuilder )
 {
 	const FVector VertexPosition = ComponentToWorldMatrix.TransformPosition( EditableMesh.GetVertexAttribute( VertexID, UEditableMeshAttribute::VertexPosition(), 0 ) );
 
-	const float DistanceToCamera = ( CameraToWorld.GetLocation() - VertexPosition ).Size();
-	const float DistanceBasedScaling = MeshEd::OverlayDistanceBias->GetFloat() + DistanceToCamera * MeshEd::OverlayDistanceScaleFactor->GetFloat();
+	const FVector CameraUpVector = CameraToWorld.TransformVector( FVector::UpVector ).GetSafeNormal();
+	const FVector CameraRightVector = CameraToWorld.TransformVector( FVector::RightVector ).GetSafeNormal();
+	const FVector CameraForwardVector = CameraToWorld.TransformVector( FVector::ForwardVector ).GetSafeNormal();
+
+	const FVector DirectionToCamera = bIsPerspectiveView ? ( CameraToWorld.GetLocation() - VertexPosition ).GetSafeNormal() : -CameraForwardVector;
+	const float DistanceToCamera = bIsPerspectiveView ? ( CameraToWorld.GetLocation() - VertexPosition ).Size() : 0.0f;
+	const float DistanceBias = bIsPerspectiveView ? MeshEd::OverlayPerspectiveDistanceBias->GetFloat() : MeshEd::OverlayOrthographicDistanceBias->GetFloat();
+	const float DistanceBasedScaling = DistanceBias + DistanceToCamera * MeshEd::OverlayDistanceScaleFactor->GetFloat();
 
 	const float SpriteSize = ( MeshEd::OverlayVertexSize->GetFloat() + SizeBias ) * DistanceBasedScaling;
 	const float HalfSpriteSize = SpriteSize * 0.5f;
 
-	const FVector CameraUpVector = CameraToWorld.TransformVector( FVector::UpVector );
-	const FVector CameraRightVector = CameraToWorld.TransformVector( FVector::RightVector );
-	const FVector DirectionToCamera = ( CameraToWorld.GetLocation() - VertexPosition ).GetSafeNormal();
-	const FVector RightDirection = CameraRightVector.GetSafeNormal();
-	const FVector UpDirection = CameraUpVector.GetSafeNormal();
-
 	// We're offsetting the geometry from the actual face a bit, to avoid z-fighting for this particular effect
 	const FVector VertexOffset = bApplyDepthBias ? ( DirectionToCamera * MeshEd::OverlayDepthOffset->GetFloat() * DistanceBasedScaling ) : FVector::ZeroVector;
 
 	FVector QuadPositions[ 4 ];
-	QuadPositions[ 0 ] = ( VertexPosition - RightDirection * HalfSpriteSize - UpDirection * HalfSpriteSize ) + VertexOffset;
-	QuadPositions[ 1 ] = ( VertexPosition - RightDirection * HalfSpriteSize + UpDirection * HalfSpriteSize ) + VertexOffset;
-	QuadPositions[ 2 ] = ( VertexPosition + RightDirection * HalfSpriteSize + UpDirection * HalfSpriteSize ) + VertexOffset;
-	QuadPositions[ 3 ] = ( VertexPosition + RightDirection * HalfSpriteSize - UpDirection * HalfSpriteSize ) + VertexOffset;
+	QuadPositions[ 0 ] = ( VertexPosition - CameraRightVector * HalfSpriteSize - CameraUpVector * HalfSpriteSize ) + VertexOffset;
+	QuadPositions[ 1 ] = ( VertexPosition - CameraRightVector * HalfSpriteSize + CameraUpVector * HalfSpriteSize ) + VertexOffset;
+	QuadPositions[ 2 ] = ( VertexPosition + CameraRightVector * HalfSpriteSize + CameraUpVector * HalfSpriteSize ) + VertexOffset;
+	QuadPositions[ 3 ] = ( VertexPosition + CameraRightVector * HalfSpriteSize - CameraUpVector * HalfSpriteSize ) + VertexOffset;
 
 	const int32 FirstVertexIndex = 
 		MeshBuilder.AddVertex( VertexOffset + QuadPositions[ 0 ], FVector2D( 0, 0 ), FVector( 1, 0, 0 ), FVector( 0, 1, 0 ), FVector( 0, 0, 1 ), ColorAndOpacity );
@@ -1953,19 +1955,21 @@ void FMeshEditorMode::AddVertexToDynamicMesh( const UEditableMesh& EditableMesh,
 }
 
 
-void FMeshEditorMode::AddThickLineToDynamicMesh( const FTransform& CameraToWorld, const FVector EdgeVertexPositions[2], const FColor ColorAndOpacity, const float SizeBias, const bool bApplyDepthBias, FDynamicMeshBuilder& MeshBuilder )
+void FMeshEditorMode::AddThickLineToDynamicMesh( const FTransform& CameraToWorld, const bool bIsPerspectiveView, const FVector EdgeVertexPositions[2], const FColor ColorAndOpacity, const float SizeBias, const bool bApplyDepthBias, FDynamicMeshBuilder& MeshBuilder )
 {
-	const float DistanceBias = MeshEd::OverlayDistanceBias->GetFloat();
-	const float Vertex0DistanceToCamera = ( CameraToWorld.GetLocation() - EdgeVertexPositions[ 0 ] ).Size();
+	const float DistanceBias = bIsPerspectiveView ? MeshEd::OverlayPerspectiveDistanceBias->GetFloat() : MeshEd::OverlayOrthographicDistanceBias->GetFloat();
+	const float Vertex0DistanceToCamera = bIsPerspectiveView ? ( CameraToWorld.GetLocation() - EdgeVertexPositions[ 0 ] ).Size() : 0.0f;
 	const float Vertex0DistanceBasedScaling = DistanceBias + Vertex0DistanceToCamera * MeshEd::OverlayDistanceScaleFactor->GetFloat();
-	const float Vertex1DistanceToCamera = ( CameraToWorld.GetLocation() - EdgeVertexPositions[ 1 ] ).Size();
+	const float Vertex1DistanceToCamera = bIsPerspectiveView ? ( CameraToWorld.GetLocation() - EdgeVertexPositions[ 1 ] ).Size() : 0.0f;
 	const float Vertex1DistanceBasedScaling = DistanceBias + Vertex1DistanceToCamera * MeshEd::OverlayDistanceScaleFactor->GetFloat();
 
 	const float Thickness = ( MeshEd::OverlayLineThickness->GetFloat() + SizeBias );
 	const float HalfThickness = Thickness * 0.5f;
 
-	const FVector DirectionToCamera0 = ( CameraToWorld.GetLocation() - EdgeVertexPositions[ 0 ] ).GetSafeNormal();
-	const FVector DirectionToCamera1 = ( CameraToWorld.GetLocation() - EdgeVertexPositions[ 1 ] ).GetSafeNormal();
+	const FVector CameraForwardVector = CameraToWorld.TransformVector( FVector::ForwardVector ).GetSafeNormal();
+
+	const FVector DirectionToCamera0 = bIsPerspectiveView ? ( CameraToWorld.GetLocation() - EdgeVertexPositions[ 0 ] ).GetSafeNormal() : -CameraForwardVector;
+	const FVector DirectionToCamera1 = bIsPerspectiveView ? ( CameraToWorld.GetLocation() - EdgeVertexPositions[ 1 ] ).GetSafeNormal() : -CameraForwardVector;
 	const FVector EdgeForward = ( EdgeVertexPositions[ 1 ] - EdgeVertexPositions[ 0 ] ).GetSafeNormal();
 	const FVector EdgeRight0 = FVector::CrossProduct( EdgeForward, DirectionToCamera0 ).GetSafeNormal();
 	const FVector EdgeRight1 = FVector::CrossProduct( EdgeForward, DirectionToCamera1 ).GetSafeNormal();
@@ -1992,7 +1996,7 @@ void FMeshEditorMode::AddThickLineToDynamicMesh( const FTransform& CameraToWorld
 }
 
 
-void FMeshEditorMode::AddEdgeToDynamicMesh( const UEditableMesh& EditableMesh, const FTransform& CameraToWorld, const FMatrix& ComponentToWorldMatrix, const FEdgeID EdgeID, const FColor ColorAndOpacity, const float SizeBias, FDynamicMeshBuilder& MeshBuilder )
+void FMeshEditorMode::AddEdgeToDynamicMesh( const UEditableMesh& EditableMesh, const FTransform& CameraToWorld, const bool bIsPerspectiveView, const FMatrix& ComponentToWorldMatrix, const FEdgeID EdgeID, const FColor ColorAndOpacity, const float SizeBias, FDynamicMeshBuilder& MeshBuilder )
 {
 	FVertexID MeshVertexIDs[ 2 ];
 	EditableMesh.GetEdgeVertices( EdgeID, /* Out */ MeshVertexIDs[0], /* Out */ MeshVertexIDs[1] );
@@ -2007,11 +2011,11 @@ void FMeshEditorMode::AddEdgeToDynamicMesh( const UEditableMesh& EditableMesh, c
 	}
 
 	const bool bApplyDepthBias = true;
-	AddThickLineToDynamicMesh( CameraToWorld, EdgeVertexPositions, ColorAndOpacity, SizeBias, bApplyDepthBias, MeshBuilder );
+	AddThickLineToDynamicMesh( CameraToWorld, bIsPerspectiveView, EdgeVertexPositions, ColorAndOpacity, SizeBias, bApplyDepthBias, MeshBuilder );
 }
 
 
-void FMeshEditorMode::AddPolygonToDynamicMesh( const UEditableMesh& EditableMesh, const FTransform& CameraToWorld, const FMatrix& ComponentToWorldMatrix, const FPolygonRef PolygonRef, const FColor ColorAndOpacity, const float SizeBias, const bool bFillFaces, FDynamicMeshBuilder& VertexAndEdgeMeshBuilder, FDynamicMeshBuilder* PolygonFaceMeshBuilder )
+void FMeshEditorMode::AddPolygonToDynamicMesh( const UEditableMesh& EditableMesh, const FTransform& CameraToWorld, const bool bIsPerspectiveView, const FMatrix& ComponentToWorldMatrix, const FPolygonRef PolygonRef, const FColor ColorAndOpacity, const float SizeBias, const bool bFillFaces, FDynamicMeshBuilder& VertexAndEdgeMeshBuilder, FDynamicMeshBuilder* PolygonFaceMeshBuilder )
 {
 	static TArray<FVertexID> MeshVertexIDs;
 	MeshVertexIDs.Reset();
@@ -2032,7 +2036,11 @@ void FMeshEditorMode::AddPolygonToDynamicMesh( const UEditableMesh& EditableMesh
 	const FVector PolygonNormal = FVector::CrossProduct( ( PolygonPerimeterVertexPositions[ 2 ] - PolygonPerimeterVertexPositions[ 0 ] ), ( PolygonPerimeterVertexPositions[ 1 ] - PolygonPerimeterVertexPositions[ 0 ] ) ).GetSafeNormal();
 	const float Determinant = ComponentToWorldMatrix.Determinant();
 
-	if( FVector::DotProduct( PolygonNormal, CameraToWorld.GetLocation() - PolygonPerimeterVertexPositions[ 0 ] ) * Determinant < 0.0f )
+	const FVector CameraForwardVectorUnnormalized = CameraToWorld.TransformVector( FVector::ForwardVector );
+	const FVector FirstVertexDirectionToCameraUnnormalized = bIsPerspectiveView ? ( CameraToWorld.GetLocation() - PolygonPerimeterVertexPositions[0] ) : -CameraForwardVectorUnnormalized;
+
+	// @todo mesheditor: Backface culling in orthographic views may not always be desirable.  Make optional?
+	if( FVector::DotProduct( PolygonNormal, FirstVertexDirectionToCameraUnnormalized ) * Determinant < 0.0f )
 	{
 		// Ignore backfaced polys
 		return;
@@ -2040,6 +2048,8 @@ void FMeshEditorMode::AddPolygonToDynamicMesh( const UEditableMesh& EditableMesh
 
 	if( bFillFaces )
 	{
+		const float DistanceBias = bIsPerspectiveView ? MeshEd::OverlayPerspectiveDistanceBias->GetFloat() : MeshEd::OverlayOrthographicDistanceBias->GetFloat();
+
 		check( PolygonFaceMeshBuilder );
 		const uint32 PolygonTriangleCount = EditableMesh.GetPolygonTriangulatedTriangleCount( PolygonRef );
 		for( uint32 PolygonTriangleNumber = 0; PolygonTriangleNumber < PolygonTriangleCount; ++PolygonTriangleNumber )
@@ -2053,10 +2063,10 @@ void FMeshEditorMode::AddPolygonToDynamicMesh( const UEditableMesh& EditableMesh
 						EditableMesh.GetPolygonTriangulatedTriangleVertexPosition( PolygonRef, PolygonTriangleNumber, TriangleVertexNumber ) );
 
 				// We're offsetting the geometry from the actual face a bit, to avoid z-fighting for this particular effect
-				const FVector DirectionToCamera = CameraToWorld.GetLocation() - TriangleVertexPositions[ TriangleVertexNumber ];
-				const float DistanceToCamera = DirectionToCamera.Size();
-				const float DistanceBasedScaling = MeshEd::OverlayDistanceBias->GetFloat() + DistanceToCamera * MeshEd::OverlayDistanceScaleFactor->GetFloat();
-				VertexOffsets[ TriangleVertexNumber ] = ( DirectionToCamera / DistanceToCamera ) * MeshEd::OverlayDepthOffset->GetFloat() * DistanceBasedScaling;
+				const FVector DirectionToCamera = bIsPerspectiveView ? ( CameraToWorld.GetLocation() - TriangleVertexPositions[ TriangleVertexNumber ] ).GetSafeNormal() : -CameraForwardVectorUnnormalized.GetSafeNormal();
+				const float DistanceToCamera = bIsPerspectiveView ? ( CameraToWorld.GetLocation() - TriangleVertexPositions[ TriangleVertexNumber ] ).Size() : 0.0f;
+				const float DistanceBasedScaling = DistanceBias + DistanceToCamera * MeshEd::OverlayDistanceScaleFactor->GetFloat();
+				VertexOffsets[ TriangleVertexNumber ] = DirectionToCamera * MeshEd::OverlayDepthOffset->GetFloat() * DistanceBasedScaling;
 			}
 
 			// @todo mesheditor perf: Horribly inefficient way to draw a single triangle!  Also we should have a global utility function for this rather than using FDynamicMeshBuilder
@@ -2079,13 +2089,13 @@ void FMeshEditorMode::AddPolygonToDynamicMesh( const UEditableMesh& EditableMesh
 		// @todo mesheditor: Assumes no holes, etc.
 		for( const FEdgeID EdgeID : PerimeterEdges )
 		{
-			AddEdgeToDynamicMesh( EditableMesh, CameraToWorld, ComponentToWorldMatrix, EdgeID, ColorAndOpacity, SizeBias, VertexAndEdgeMeshBuilder );
+			AddEdgeToDynamicMesh( EditableMesh, CameraToWorld, bIsPerspectiveView, ComponentToWorldMatrix, EdgeID, ColorAndOpacity, SizeBias, VertexAndEdgeMeshBuilder );
 		}
 	}
 }
 
 
-void FMeshEditorMode::DrawMeshElements( const FTransform& CameraToWorld, FViewport* Viewport, FPrimitiveDrawInterface* PDI, const TArrayView<FMeshElement>& MeshElements, const FColor Color, const bool bFillFaces, const float SizeBias, const TArray<FColor>* OptionalPerElementColors, const TArray<float>* OptionalPerElementSizeBiases )
+void FMeshEditorMode::DrawMeshElements( const FTransform& CameraToWorld, const bool bIsPerspectiveView, FViewport* Viewport, FPrimitiveDrawInterface* PDI, const TArrayView<FMeshElement>& MeshElements, const FColor Color, const bool bFillFaces, const float SizeBias, const TArray<FColor>* OptionalPerElementColors, const TArray<float>* OptionalPerElementSizeBiases )
 {
 	if( MeshElements.Num() > 0 )
 	{
@@ -2122,6 +2132,7 @@ void FMeshEditorMode::DrawMeshElements( const FTransform& CameraToWorld, FViewpo
 							AddVertexToDynamicMesh(
 								*EditableMesh,
 								CameraToWorld,
+								bIsPerspectiveView,
 								ComponentToWorldMatrix,
 								FVertexID( MeshElement.ElementAddress.ElementID ),
 								ElementColor,
@@ -2134,6 +2145,7 @@ void FMeshEditorMode::DrawMeshElements( const FTransform& CameraToWorld, FViewpo
 							AddEdgeToDynamicMesh(
 								*EditableMesh,
 								CameraToWorld,
+								bIsPerspectiveView,
 								ComponentToWorldMatrix,
 								FEdgeID( MeshElement.ElementAddress.ElementID ),
 								ElementColor,
@@ -2147,6 +2159,7 @@ void FMeshEditorMode::DrawMeshElements( const FTransform& CameraToWorld, FViewpo
 							AddPolygonToDynamicMesh(
 								*EditableMesh,
 								CameraToWorld,
+								bIsPerspectiveView,
 								ComponentToWorldMatrix,
 								FPolygonRef( MeshElement.ElementAddress.SectionID, FPolygonID( MeshElement.ElementAddress.ElementID ) ),
 								ElementColor,
@@ -2188,7 +2201,16 @@ void FMeshEditorMode::Render( const FSceneView* SceneView, FViewport* Viewport,
 {
 	FEdMode::Render( SceneView, Viewport, PDI );
 
-	const FTransform CameraToWorld = CachedCameraToWorld.IsSet() ? CachedCameraToWorld.GetValue() : FTransform( SceneView->ViewRotation, SceneView->ViewLocation );
+	FTransform CameraToWorld;
+	{
+	    // The SceneView's view matrices are rotated 90 degrees in pitch, so we'll have to undo that to get
+		// back to Unreal standard coordinates
+	    const FMatrix& WorldToCameraMatrix = SceneView->ViewMatrices.GetViewMatrix();
+	    const FMatrix Rotate90( FPlane( 0, 0, 1, 0 ), FPlane( 1, 0, 0, 0 ), FPlane( 0, 1, 0, 0 ), FPlane( 0, 0, 0, 1 ) );
+		CameraToWorld.SetFromMatrix( ( WorldToCameraMatrix * Rotate90.Inverse() ).Inverse() );
+	}
+
+	const bool bIsPerspectiveView = SceneView->IsPerspectiveProjection();
 
 	// @todo mesheditor debug
 	if( MeshEd::ShowDebugStats->GetInt() > 0 && SelectedMeshElements.Num() > 0 )
@@ -2328,7 +2350,7 @@ void FMeshEditorMode::Render( const FSceneView* SceneView, FViewport* Viewport,
 					}					
 
 					const bool bApplyDepthBias = true;
-					AddThickLineToDynamicMesh( CameraToWorld, WorldSpaceEdgeVertexPositions, ColorAndOpacity, SizeBias, bApplyDepthBias, VertexAndEdgeMeshBuilder );
+					AddThickLineToDynamicMesh( CameraToWorld, bIsPerspectiveView, WorldSpaceEdgeVertexPositions, ColorAndOpacity, SizeBias, bApplyDepthBias, VertexAndEdgeMeshBuilder );
 				}
 			}
 
@@ -2356,6 +2378,7 @@ void FMeshEditorMode::Render( const FSceneView* SceneView, FViewport* Viewport,
 								AddPolygonToDynamicMesh(
 									*EditableMesh,
 									CameraToWorld,
+									bIsPerspectiveView,
 									ComponentToWorldMatrix,
 									PolygonRef,
 									ColorAndOpacity,
@@ -2385,6 +2408,7 @@ void FMeshEditorMode::Render( const FSceneView* SceneView, FViewport* Viewport,
 										const bool bApplyDepthBias = false;
 										AddThickLineToDynamicMesh(
 											CameraToWorld,
+											bIsPerspectiveView,
 											VertexNormalRenderPositions,
 											FColor::Magenta,
 											SizeBias,
@@ -2437,6 +2461,7 @@ void FMeshEditorMode::Render( const FSceneView* SceneView, FViewport* Viewport,
 				const bool bFillFaces = true;
 				DrawMeshElements(
 					CameraToWorld,
+					bIsPerspectiveView,
 					Viewport,
 					PDI,
 					HoveredMeshElementsToDraw,
@@ -2477,6 +2502,7 @@ void FMeshEditorMode::Render( const FSceneView* SceneView, FViewport* Viewport,
 				const bool bFillFaces = true;
 				DrawMeshElements(
 					CameraToWorld,
+					bIsPerspectiveView,
 					Viewport,
 					PDI,
 					FadingOutHoveredMeshElementsToDraw,
@@ -2507,6 +2533,7 @@ void FMeshEditorMode::Render( const FSceneView* SceneView, FViewport* Viewport,
 			const bool bFillFaces = true;
 			DrawMeshElements(
 				CameraToWorld,
+				bIsPerspectiveView,
 				Viewport,
 				PDI,
 				SelectedMeshElements,
@@ -2732,6 +2759,7 @@ void FMeshEditorMode::OnViewportInteractionHoverUpdate( UViewportInteractor* Vie
 										ComponentToWorldMatrix.InverseTransformVector( FVector( MeshEditorInteractorData.GrabberSphere.W ) ).X );
 
 									const FTransform CameraToWorld = CachedCameraToWorld.IsSet() ? CachedCameraToWorld.GetValue() : HitComponent->GetComponentToWorld();	// @todo mesheditor: We're expecting to have a CameraToWorld here
+									const bool bIsPerspectiveView = bCachedIsPerspectiveView.IsSet() ? bCachedIsPerspectiveView.GetValue() : true;
 									const FVector ComponentSpaceCameraLocation = ComponentToWorldMatrix.InverseTransformPosition( CameraToWorld.GetLocation() );
 									const float ComponentSpaceFuzzyDistanceScaleFactor = ComponentToWorldMatrix.InverseTransformVector( FVector( MeshEd::OverlayDistanceScaleFactor->GetFloat() / ViewportWorldInteraction->GetWorldScaleFactor(), 0.0f, 0.0f ) ).Size();
 
@@ -2748,6 +2776,7 @@ void FMeshEditorMode::OnViewportInteractionHoverUpdate( UViewportInteractor* Vie
 										ComponentSpaceRayFuzzyDistance, 
 										OnlyElementType, 
 										ComponentSpaceCameraLocation,
+										bIsPerspectiveView,
 										ComponentSpaceFuzzyDistanceScaleFactor,
 										/* Out */ HitInteractorShape,
 										/* Out */ ComponentSpaceHitLocation );
@@ -3195,7 +3224,8 @@ void FMeshEditorMode::UpdateActiveAction( const bool bIsActionFinishing )
 					// Check if the polygon normal is pointing towards us. If not, we need to flip the polygon
 					FVector PolygonNormal = EditableMesh->ComputePolygonNormal( NewPolygonRefs[ 0 ] );
 
-					if( CachedCameraToWorld.IsSet() )
+					// @todo mesheditor: Add support for backface checks in orthographic mode
+					if( CachedCameraToWorld.IsSet() && bCachedIsPerspectiveView.IsSet() && !bCachedIsPerspectiveView.GetValue() )
 					{
 						if( FVector::DotProduct( Component->GetComponentTransform().TransformVector( PolygonNormal ), DrawnPoints[ 0 ].Get<1>() - CachedCameraToWorld.GetValue().GetLocation() ) > 0.0f )
 						{
@@ -3505,7 +3535,7 @@ void FMeshEditorMode::FindEdgeSplitUnderInteractor(	UViewportInteractor* Viewpor
 }
 
 
-FEditableMeshElementAddress FMeshEditorMode::QueryElement( const UEditableMesh& EditableMesh, const bool bUseSphere, const FSphere& Sphere, const float SphereFuzzyDistance, bool bUseRay, const FVector& RayStart, const FVector& RayEnd, const float RayFuzzyDistance, const EEditableMeshElementType OnlyElementType, const FVector& CameraLocation, const float FuzzyDistanceScaleFactor, EInteractorShape& OutInteractorShape, FVector& OutHitLocation ) const
+FEditableMeshElementAddress FMeshEditorMode::QueryElement( const UEditableMesh& EditableMesh, const bool bUseSphere, const FSphere& Sphere, const float SphereFuzzyDistance, bool bUseRay, const FVector& RayStart, const FVector& RayEnd, const float RayFuzzyDistance, const EEditableMeshElementType OnlyElementType, const FVector& CameraLocation, const bool bIsPerspectiveView, const float FuzzyDistanceScaleFactor, EInteractorShape& OutInteractorShape, FVector& OutHitLocation ) const
 {
 	OutHitLocation = FVector::ZeroVector;
 
@@ -3537,7 +3567,7 @@ FEditableMeshElementAddress FMeshEditorMode::QueryElement( const UEditableMesh&
 				{
 					const FVector PolygonNormal = EditableMesh.ComputePolygonNormal( PolygonRef );
 					const FVector PolygonCenter = EditableMesh.ComputePolygonCenter( PolygonRef );
-					if( FVector::DotProduct( CameraLocation - PolygonCenter, PolygonNormal ) > 0.0f )
+					if( !bIsPerspectiveView || FVector::DotProduct( CameraLocation - PolygonCenter, PolygonNormal ) > 0.0f )	// @todo mesheditor: Add support for backface culling in orthographic views
 					{
 						FrontFacingPolygons.Add( PolygonRef );
 
@@ -3578,7 +3608,7 @@ FEditableMeshElementAddress FMeshEditorMode::QueryElement( const UEditableMesh&
 			}
 
 			const bool bAlreadyHitTriangle = ( HitElementAddress.ElementType == EEditableMeshElementType::Polygon );
-			const bool bHit = CheckTriangle( bUseSphere, Sphere, SphereFuzzyDistance, bUseRay, RayStart, CurrentRayEnd, RayFuzzyDistance, TriangleVertexPositions, CameraLocation, FuzzyDistanceScaleFactor, ClosestInteractorShape, ClosestDistanceToRay, ClosestDistanceOnRay, ClosestHitLocation, bAlreadyHitTriangle );
+			const bool bHit = CheckTriangle( bUseSphere, Sphere, SphereFuzzyDistance, bUseRay, RayStart, CurrentRayEnd, RayFuzzyDistance, TriangleVertexPositions, CameraLocation, bIsPerspectiveView, FuzzyDistanceScaleFactor, ClosestInteractorShape, ClosestDistanceToRay, ClosestDistanceOnRay, ClosestHitLocation, bAlreadyHitTriangle );
 			if( bHit )
 			{
 				HitElementAddress.ElementType = EEditableMeshElementType::Polygon;
@@ -3601,7 +3631,7 @@ FEditableMeshElementAddress FMeshEditorMode::QueryElement( const UEditableMesh&
 			EdgeVertexPositions[ 1 ] = EditableMesh.GetVertexAttribute( EditableMesh.GetEdgeVertex( EdgeID, 1 ), UEditableMeshAttribute::VertexPosition(), 0 );
 
 			const bool bAlreadyHitEdge = ( HitElementAddress.ElementType == EEditableMeshElementType::Edge );
-			const bool bHit = CheckEdge( bUseSphere, Sphere, SphereFuzzyDistance, bUseRay, RayStart, CurrentRayEnd, RayFuzzyDistance, EdgeVertexPositions, CameraLocation, FuzzyDistanceScaleFactor, ClosestInteractorShape, ClosestDistanceToRay, ClosestDistanceOnRay, ClosestHitLocation, bAlreadyHitEdge );
+			const bool bHit = CheckEdge( bUseSphere, Sphere, SphereFuzzyDistance, bUseRay, RayStart, CurrentRayEnd, RayFuzzyDistance, EdgeVertexPositions, CameraLocation, bIsPerspectiveView, FuzzyDistanceScaleFactor, ClosestInteractorShape, ClosestDistanceToRay, ClosestDistanceOnRay, ClosestHitLocation, bAlreadyHitEdge );
 			if( bHit )
 			{
 				HitElementAddress.ElementType = EEditableMeshElementType::Edge;
@@ -3622,7 +3652,7 @@ FEditableMeshElementAddress FMeshEditorMode::QueryElement( const UEditableMesh&
 		{
 			const FVector VertexPosition = EditableMesh.GetVertexAttribute( VertexID, UEditableMeshAttribute::VertexPosition(), 0 );
 			const bool bAlreadyHitVertex = ( HitElementAddress.ElementType == EEditableMeshElementType::Vertex );
-			const bool bHit = CheckVertex( bUseSphere, Sphere, SphereFuzzyDistance, bUseRay, RayStart, CurrentRayEnd, RayFuzzyDistance, VertexPosition, CameraLocation, FuzzyDistanceScaleFactor, ClosestInteractorShape, ClosestDistanceToRay, ClosestDistanceOnRay, ClosestHitLocation, bAlreadyHitVertex );
+			const bool bHit = CheckVertex( bUseSphere, Sphere, SphereFuzzyDistance, bUseRay, RayStart, CurrentRayEnd, RayFuzzyDistance, VertexPosition, CameraLocation, bIsPerspectiveView, FuzzyDistanceScaleFactor, ClosestInteractorShape, ClosestDistanceToRay, ClosestDistanceOnRay, ClosestHitLocation, bAlreadyHitVertex );
 			if( bHit )
 			{
 				HitElementAddress.ElementType = EEditableMeshElementType::Vertex;
@@ -3642,12 +3672,13 @@ FEditableMeshElementAddress FMeshEditorMode::QueryElement( const UEditableMesh&
 }
 
 
-bool FMeshEditorMode::CheckVertex( const bool bUseSphere, const FSphere& Sphere, const float SphereFuzzyDistance, bool bUseRay, const FVector& RayStart, const FVector& RayEnd, const float RayFuzzyDistance, const FVector& VertexPosition, const FVector& CameraLocation, const float FuzzyDistanceScaleFactor, EInteractorShape& ClosestInteractorShape, float& ClosestDistanceToRay, float& ClosestDistanceOnRay, FVector& ClosestHitLocation, const bool bAlreadyHitVertex )
+bool FMeshEditorMode::CheckVertex( const bool bUseSphere, const FSphere& Sphere, const float SphereFuzzyDistance, bool bUseRay, const FVector& RayStart, const FVector& RayEnd, const float RayFuzzyDistance, const FVector& VertexPosition, const FVector& CameraLocation, const bool bIsPerspectiveView, const float FuzzyDistanceScaleFactor, EInteractorShape& ClosestInteractorShape, float& ClosestDistanceToRay, float& ClosestDistanceOnRay, FVector& ClosestHitLocation, const bool bAlreadyHitVertex )
 {
 	bool bHit = false;
 
- 	const float DistanceToCamera = ( CameraLocation - VertexPosition ).Size();
-	const float DistanceBasedScaling = MeshEd::OverlayDistanceBias->GetFloat() + DistanceToCamera * FuzzyDistanceScaleFactor;
+	const float DistanceBias = bIsPerspectiveView ? MeshEd::OverlayPerspectiveDistanceBias->GetFloat() : MeshEd::OverlayOrthographicDistanceBias->GetFloat();
+	const float DistanceToCamera = bIsPerspectiveView ? ( CameraLocation - VertexPosition ).Size() : 0.0f;
+	const float DistanceBasedScaling = DistanceBias + DistanceToCamera * FuzzyDistanceScaleFactor;
 	check( DistanceBasedScaling > 0.0f );
 
 	if( bUseSphere )
@@ -3705,7 +3736,7 @@ bool FMeshEditorMode::CheckVertex( const bool bUseSphere, const FSphere& Sphere,
 }
 
 
-bool FMeshEditorMode::CheckEdge( const bool bUseSphere, const FSphere& Sphere, const float SphereFuzzyDistance, bool bUseRay, const FVector& RayStart, const FVector& RayEnd, const float RayFuzzyDistance, const FVector EdgeVertexPositions[ 2 ], const FVector& CameraLocation, const float FuzzyDistanceScaleFactor, EInteractorShape& ClosestInteractorShape, float& ClosestDistanceToRay, float& ClosestDistanceOnRay, FVector& ClosestHitLocation, const bool bAlreadyHitEdge )
+bool FMeshEditorMode::CheckEdge( const bool bUseSphere, const FSphere& Sphere, const float SphereFuzzyDistance, bool bUseRay, const FVector& RayStart, const FVector& RayEnd, const float RayFuzzyDistance, const FVector EdgeVertexPositions[ 2 ], const FVector& CameraLocation, const bool bIsPerspectiveView, const float FuzzyDistanceScaleFactor, EInteractorShape& ClosestInteractorShape, float& ClosestDistanceToRay, float& ClosestDistanceOnRay, FVector& ClosestHitLocation, const bool bAlreadyHitEdge )
 {
 	bool bHit = false;
 
@@ -3715,8 +3746,9 @@ bool FMeshEditorMode::CheckEdge( const bool bUseSphere, const FSphere& Sphere, c
 		if( DistanceToSphere <= Sphere.W )
 		{
 			const FVector ClosestPointOnEdge = FMath::ClosestPointOnSegment( Sphere.Center, EdgeVertexPositions[ 0 ], EdgeVertexPositions[ 1 ] );
-			const float DistanceToCamera = ( CameraLocation - ClosestPointOnEdge ).Size();
-			const float DistanceBasedScaling = MeshEd::OverlayDistanceBias->GetFloat() + DistanceToCamera * FuzzyDistanceScaleFactor;
+			const float DistanceToCamera = bIsPerspectiveView ? ( CameraLocation - ClosestPointOnEdge ).Size() : 0.0f;
+			const float DistanceBias = bIsPerspectiveView ? MeshEd::OverlayPerspectiveDistanceBias->GetFloat() : MeshEd::OverlayOrthographicDistanceBias->GetFloat();
+			const float DistanceBasedScaling = DistanceBias + DistanceToCamera * FuzzyDistanceScaleFactor;
 
 			// Inside sphere
 			if( DistanceToSphere < ClosestDistanceToRay ||
@@ -3749,8 +3781,9 @@ bool FMeshEditorMode::CheckEdge( const bool bUseSphere, const FSphere& Sphere, c
 		const bool bIsBehindRay = FVector::DotProduct( RayDirection, DirectionTowardClosestPointOnRay ) < 0.0f;
 		if( !bIsBehindRay )
 		{
-			const float DistanceToCamera = ( CameraLocation - ClosestPointOnEdge ).Size();
-			const float DistanceBasedScaling = MeshEd::OverlayDistanceBias->GetFloat() + DistanceToCamera * FuzzyDistanceScaleFactor;
+			const float DistanceToCamera = bIsPerspectiveView ? ( CameraLocation - ClosestPointOnEdge ).Size() : 0.0f;
+			const float DistanceBias = bIsPerspectiveView ? MeshEd::OverlayPerspectiveDistanceBias->GetFloat() : MeshEd::OverlayOrthographicDistanceBias->GetFloat();
+			const float DistanceBasedScaling = DistanceBias + DistanceToCamera * FuzzyDistanceScaleFactor;
 			check( DistanceBasedScaling > 0.0f );
 
 			// Are we within the minimum distance for hitting an edge?
@@ -3778,7 +3811,7 @@ bool FMeshEditorMode::CheckEdge( const bool bUseSphere, const FSphere& Sphere, c
 }
 
 
-bool FMeshEditorMode::CheckTriangle( const bool bUseSphere, const FSphere& Sphere, const float SphereFuzzyDistance, bool bUseRay, const FVector& RayStart, const FVector& RayEnd, const float RayFuzzyDistance, const FVector TriangleVertexPositions[ 3 ], const FVector& CameraLocation, const float FuzzyDistanceScaleFactor, EInteractorShape& ClosestInteractorShape, float& ClosestDistanceToRay, float& ClosestDistanceOnRay, FVector& ClosestHitLocation, const bool bAlreadyHitTriangle )
+bool FMeshEditorMode::CheckTriangle( const bool bUseSphere, const FSphere& Sphere, const float SphereFuzzyDistance, bool bUseRay, const FVector& RayStart, const FVector& RayEnd, const float RayFuzzyDistance, const FVector TriangleVertexPositions[ 3 ], const FVector& CameraLocation, const bool bIsPerspectiveView, const float FuzzyDistanceScaleFactor, EInteractorShape& ClosestInteractorShape, float& ClosestDistanceToRay, float& ClosestDistanceOnRay, FVector& ClosestHitLocation, const bool bAlreadyHitTriangle )
 {
 	bool bHit = false;
 
@@ -3789,8 +3822,9 @@ bool FMeshEditorMode::CheckTriangle( const bool bUseSphere, const FSphere& Spher
 		const float DistanceToSphere = ( ClosestPointOnTriangleToSphere - Sphere.Center ).Size();
 		if( DistanceToSphere <= Sphere.W )
 		{
-			const float DistanceToCamera = ( CameraLocation - ClosestPointOnTriangleToSphere ).Size();
-			const float DistanceBasedScaling = MeshEd::OverlayDistanceBias->GetFloat() + DistanceToCamera * FuzzyDistanceScaleFactor;
+			const float DistanceToCamera = bIsPerspectiveView ? ( CameraLocation - ClosestPointOnTriangleToSphere ).Size() : 0.0f;
+			const float DistanceBias = bIsPerspectiveView ? MeshEd::OverlayPerspectiveDistanceBias->GetFloat() : MeshEd::OverlayOrthographicDistanceBias->GetFloat();
+			const float DistanceBasedScaling = DistanceBias + DistanceToCamera * FuzzyDistanceScaleFactor;
 
 			// Inside sphere
 			if( DistanceToSphere < ClosestDistanceToRay ||
@@ -3845,8 +3879,9 @@ bool FMeshEditorMode::CheckTriangle( const bool bUseSphere, const FSphere& Spher
 		// @todo mesheditor perf: We also have a SIMD version of this that does four triangles at once
 		if( Local::RayIntersectTriangle( RayStart, RayEnd, TriangleVertexPositions[ 0 ], TriangleVertexPositions[ 1 ], TriangleVertexPositions[ 2 ], /* Out */ IntersectionPoint ) )
 		{
-			const float DistanceToCamera = ( CameraLocation - IntersectionPoint ).Size();
-			const float DistanceBasedScaling = MeshEd::OverlayDistanceBias->GetFloat() + DistanceToCamera * FuzzyDistanceScaleFactor;
+			const float DistanceToCamera = bIsPerspectiveView ? ( CameraLocation - IntersectionPoint ).Size() : 0.0f;
+			const float DistanceBias = bIsPerspectiveView ? MeshEd::OverlayPerspectiveDistanceBias->GetFloat() : MeshEd::OverlayOrthographicDistanceBias->GetFloat();
+			const float DistanceBasedScaling = DistanceBias + DistanceToCamera * FuzzyDistanceScaleFactor;
     
 			const float DistanceToRay = 0.0f;  // We intersected the triangle, otherwise we wouldn't even be in here
 			const float DistanceOnRay = ( IntersectionPoint - RayStart ).Size();
diff --git a/Engine/Plugins/Editor/MeshEditor/Source/MeshEditor/MeshEditorMode.h b/Engine/Plugins/Editor/MeshEditor/Source/MeshEditor/MeshEditorMode.h
--- a/Engine/Plugins/Editor/MeshEditor/Source/MeshEditor/MeshEditorMode.h
+++ b/Engine/Plugins/Editor/MeshEditor/Source/MeshEditor/MeshEditorMode.h
@@ -207,19 +207,19 @@ protected:
 	const class UMeshEditorAssetContainer& GetAssetContainer() const;
 
 	/** Fills the specified dynamic mesh builder with primitives to render a mesh vertex */
-	void AddVertexToDynamicMesh( const UEditableMesh& EditableMesh, const FTransform& CameraToWorld, const FMatrix& ComponentToWorldMatrix, const FVertexID VertexID, const FColor ColorAndOpacity, const float SizeBias, const bool bApplyDepthBias, class FDynamicMeshBuilder& MeshBuilder );
+	void AddVertexToDynamicMesh( const UEditableMesh& EditableMesh, const FTransform& CameraToWorld, const bool bIsPerspectiveView, const FMatrix& ComponentToWorldMatrix, const FVertexID VertexID, const FColor ColorAndOpacity, const float SizeBias, const bool bApplyDepthBias, class FDynamicMeshBuilder& MeshBuilder );
 		
 	/** Fills the specified dynamic mesh builder with primitives to a thick line.  Incoming positions are in world space. */
-	void AddThickLineToDynamicMesh( const FTransform& CameraToWorld, const FVector EdgeVertexPositions[2], const FColor ColorAndOpacity, const float SizeBias, const bool bApplyDepthBias, FDynamicMeshBuilder& MeshBuilder );
+	void AddThickLineToDynamicMesh( const FTransform& CameraToWorld, const bool bIsPerspectiveView, const FVector EdgeVertexPositions[2], const FColor ColorAndOpacity, const float SizeBias, const bool bApplyDepthBias, FDynamicMeshBuilder& MeshBuilder );
 
 	/** Fills the specified dynamic mesh builder with primitives to render a mesh edge */
-	void AddEdgeToDynamicMesh( const UEditableMesh& EditableMesh, const FTransform& CameraToWorld, const FMatrix& ComponentToWorldMatrix, const FEdgeID EdgeID, const FColor ColorAndOpacity, const float SizeBias, class FDynamicMeshBuilder& MeshBuilder );
+	void AddEdgeToDynamicMesh( const UEditableMesh& EditableMesh, const FTransform& CameraToWorld, const bool bIsPerspectiveView, const FMatrix& ComponentToWorldMatrix, const FEdgeID EdgeID, const FColor ColorAndOpacity, const float SizeBias, class FDynamicMeshBuilder& MeshBuilder );
 
 	/** Fills the specified dynamic mesh builders with primitives to render a polygon and it's edges */
-	void AddPolygonToDynamicMesh( const UEditableMesh& EditableMesh, const FTransform& CameraToWorld, const FMatrix& ComponentToWorldMatrix, const FPolygonRef PolygonRef, const FColor ColorAndOpacity, const float SizeBias, const bool bFillFaces, class FDynamicMeshBuilder& VertexAndEdgeMeshBuilder, class FDynamicMeshBuilder* PolygonFaceMeshBuilder );
+	void AddPolygonToDynamicMesh( const UEditableMesh& EditableMesh, const FTransform& CameraToWorld, const bool bIsPerspectiveView, const FMatrix& ComponentToWorldMatrix, const FPolygonRef PolygonRef, const FColor ColorAndOpacity, const float SizeBias, const bool bFillFaces, class FDynamicMeshBuilder& VertexAndEdgeMeshBuilder, class FDynamicMeshBuilder* PolygonFaceMeshBuilder );
 
 	/** Renders the spceified mesh element */
-	void DrawMeshElements( const FTransform& CameraToWorld, FViewport* Viewport, FPrimitiveDrawInterface* PDI, const TArrayView<FMeshElement>& MeshElements, const FColor Color, const bool bFillFaces, const float HoverAnimation, const TArray<FColor>* OptionalPerElementColors = nullptr, const TArray<float>* OptionalPerElementSizeBiases = nullptr );
+	void DrawMeshElements( const FTransform& CameraToWorld, const bool bIsPerspectiveView, FViewport* Viewport, FPrimitiveDrawInterface* PDI, const TArrayView<FMeshElement>& MeshElements, const FColor Color, const bool bFillFaces, const float HoverAnimation, const TArray<FColor>* OptionalPerElementColors = nullptr, const TArray<float>* OptionalPerElementSizeBiases = nullptr );
 
 	/** Called every frame for each viewport interactor to update what's under the cursor */
 	void OnViewportInteractionHoverUpdate( class UViewportInteractor* ViewportInteractor, FVector& OutHoverImpactPoint, bool& bWasHandled );
@@ -246,10 +246,10 @@ protected:
 	void UpdateActiveAction( const bool bIsActionFinishing );
 
 	/** Geometry tests */
-	FEditableMeshElementAddress QueryElement( const UEditableMesh& EditableMesh, const bool bUseSphere, const FSphere& Sphere, const float SphereFuzzyDistance, bool bUseRay, const FVector& RayStart, const FVector& RayEnd, const float RayFuzzyDistance, const EEditableMeshElementType OnlyElementType, const FVector& CameraLocation, const float FuzzyDistanceScaleFactor, EInteractorShape& OutInteractorShape, FVector& OutHitLocation ) const;
-	static bool CheckVertex( const bool bUseSphere, const FSphere& Sphere, const float SphereFuzzyDistance, bool bUseRay, const FVector& RayStart, const FVector& RayEnd, const float FuzzyDistance, const FVector& VertexPosition, const FVector& CameraLocation, const float FuzzyDistanceScaleFactor, EInteractorShape& ClosestInteractorShape, float& ClosestDistanceToRay, float& ClosestDistanceOnRay, FVector& ClosestHitLocation, const bool bAlreadyHitVertex );
-	static bool CheckEdge( const bool bUseSphere, const FSphere& Sphere, const float SphereFuzzyDistance, bool bUseRay, const FVector& RayStart, const FVector& RayEnd, const float FuzzyDistance, const FVector EdgeVertexPositions[ 2 ], const FVector& CameraLocation, const float FuzzyDistanceScaleFactor, EInteractorShape& ClosestInteractorShape, float& ClosestDistanceToRay, float& ClosestDistanceOnRay, FVector& ClosestHitLocation, const bool bAlreadyEdge );
-	static bool CheckTriangle( const bool bUseSphere, const FSphere& Sphere, const float SphereFuzzyDistance, bool bUseRay, const FVector& RayStart, const FVector& RayEnd, const float FuzzyDistance, const FVector TriangleVertexPositions[ 3 ], const FVector& CameraLocation, const float FuzzyDistanceScaleFactor, EInteractorShape& ClosestInteractorShape, float& ClosestDistanceToRay, float& ClosestDistanceOnRay, FVector& ClosestHitLocation, const bool bAlreadyHitTriangle );
+	FEditableMeshElementAddress QueryElement( const UEditableMesh& EditableMesh, const bool bUseSphere, const FSphere& Sphere, const float SphereFuzzyDistance, bool bUseRay, const FVector& RayStart, const FVector& RayEnd, const float RayFuzzyDistance, const EEditableMeshElementType OnlyElementType, const FVector& CameraLocation, const bool bIsPerspectiveView, const float FuzzyDistanceScaleFactor, EInteractorShape& OutInteractorShape, FVector& OutHitLocation ) const;
+	static bool CheckVertex( const bool bUseSphere, const FSphere& Sphere, const float SphereFuzzyDistance, bool bUseRay, const FVector& RayStart, const FVector& RayEnd, const float FuzzyDistance, const FVector& VertexPosition, const FVector& CameraLocation, const bool bIsPerspectiveView, const float FuzzyDistanceScaleFactor, EInteractorShape& ClosestInteractorShape, float& ClosestDistanceToRay, float& ClosestDistanceOnRay, FVector& ClosestHitLocation, const bool bAlreadyHitVertex );
+	static bool CheckEdge( const bool bUseSphere, const FSphere& Sphere, const float SphereFuzzyDistance, bool bUseRay, const FVector& RayStart, const FVector& RayEnd, const float FuzzyDistance, const FVector EdgeVertexPositions[ 2 ], const FVector& CameraLocation, const bool bIsPerspectiveView, const float FuzzyDistanceScaleFactor, EInteractorShape& ClosestInteractorShape, float& ClosestDistanceToRay, float& ClosestDistanceOnRay, FVector& ClosestHitLocation, const bool bAlreadyEdge );
+	static bool CheckTriangle( const bool bUseSphere, const FSphere& Sphere, const float SphereFuzzyDistance, bool bUseRay, const FVector& RayStart, const FVector& RayEnd, const float FuzzyDistance, const FVector TriangleVertexPositions[ 3 ], const FVector& CameraLocation, const bool bIsPerspectiveView, const float FuzzyDistanceScaleFactor, EInteractorShape& ClosestInteractorShape, float& ClosestDistanceToRay, float& ClosestDistanceOnRay, FVector& ClosestHitLocation, const bool bAlreadyHitTriangle );
 
 	/** Returns the index of an element in the selection set, or INDEX_NONE if its not selected */
 	int32 GetSelectedMeshElementIndex( const FMeshElement& MeshElement ) const;
@@ -639,6 +639,9 @@ protected:
 	    This is cached from the last known viewport, or taken directly from the VR head transform, if valid */
 	TOptional<FTransform> CachedCameraToWorld;
 
+	/** Whether we're in a perspective view the last time we were updated */
+	TOptional<bool> bCachedIsPerspectiveView;
+
 	/** Requested focus to selection */
 	bool bShouldFocusToSelection;
 
