Mesh Editor: Don't force "All" gizmo mode when starting mesh editor
- This was weird to clobber the user's setting.  Also, this works around the issue where VWI is not the authority over the viewport toolbar state, which cause the gizmo to be "stuck" in All mode forever.

diff --git a/Engine/Plugins/Editor/MeshEditor/Source/MeshEditor/MeshEditorMode.cpp b/Engine/Plugins/Editor/MeshEditor/Source/MeshEditor/MeshEditorMode.cpp
--- a/Engine/Plugins/Editor/MeshEditor/Source/MeshEditor/MeshEditorMode.cpp
+++ b/Engine/Plugins/Editor/MeshEditor/Source/MeshEditor/MeshEditorMode.cpp
@@ -730,9 +730,6 @@ void FMeshEditorMode::Enter()
 	ViewportWorldInteraction->OnStopDragging().AddRaw( this, &FMeshEditorMode::OnViewportInteractionStopDragging );
 	ViewportWorldInteraction->OnFinishedMovingTransformables().AddRaw( this, &FMeshEditorMode::OnViewportInteractionFinishedMovingTransformables );
 
-	// @todo gizmo: Forcing "All" mode when you enter Geometry mode, because it is generally the most useful.
-	ViewportWorldInteraction->SetGizmoHandleType( EGizmoHandleTypes::All );
-
 	// Register our system for transforming mesh elements
 	UMeshElementTransformer* MeshElementTransformer = NewObject<UMeshElementTransformer>();
 	ViewportWorldInteraction->SetTransformer( MeshElementTransformer );
