Mesh Editor extensibility (phase 8)
- Added extensibility support for 'common' mesh commands
- Migrated Quadrangulate Mesh to a new module
- Fixed some non-unity compilation errors.  There are still a lot more.

diff --git a/Engine/Plugins/Editor/MeshEditor/Source/MeshEditor/MeshEditorCommands.cpp b/Engine/Plugins/Editor/MeshEditor/Source/MeshEditor/MeshEditorCommands.cpp
--- a/Engine/Plugins/Editor/MeshEditor/Source/MeshEditor/MeshEditorCommands.cpp
+++ b/Engine/Plugins/Editor/MeshEditor/Source/MeshEditor/MeshEditorCommands.cpp
@@ -3,6 +3,8 @@
 #include "MeshEditorCommands.h"
 #include "MeshEditorStyle.h"
 #include "IMeshEditorModeUIContract.h"
+#include "UIAction.h"
+#include "UObjectIterator.h"
 
 
 #define LOCTEXT_NAMESPACE "MeshEditorCommands"
@@ -30,10 +32,21 @@ FUIAction UMeshEditorCommand::MakeUIAction( IMeshEditorModeUIContract& MeshEdito
 			this->Execute( MeshEditorMode );
 		} ) );
 
-		UIAction = FUIAction(
-			ExecuteAction,
-			FCanExecuteAction::CreateLambda( [&MeshEditorMode, ElementType] { return MeshEditorMode.IsMeshElementTypeSelected( ElementType ); } )
-		);
+		if( ElementType == EEditableMeshElementType::Invalid )
+		{
+			// Common command
+			UIAction = FUIAction(
+				ExecuteAction,
+				FCanExecuteAction::CreateLambda( [&MeshEditorMode] { return MeshEditorMode.GetSelectedEditableMeshes().Num() > 0; } )
+			);
+		}
+		else
+		{
+			UIAction = FUIAction(
+				ExecuteAction,
+				FCanExecuteAction::CreateLambda( [&MeshEditorMode, ElementType] { return MeshEditorMode.IsMeshElementTypeSelected( ElementType ); } )
+			);
+		}
 	}
 
 	return UIAction;
@@ -64,7 +77,15 @@ void FMeshEditorCommonCommands::RegisterCommands()
 	UI_COMMAND(SetEdgeSelectionMode, "Set Edge Selection Mode", "Sets the selection mode so that only edges will be selected.", EUserInterfaceActionType::None, FInputChord(EKeys::Two));
 	UI_COMMAND(SetPolygonSelectionMode, "Set Polygon Selection Mode", "Sets the selection mode so that only polygons will be selected.", EUserInterfaceActionType::None, FInputChord(EKeys::Three));
 	UI_COMMAND(SetAnySelectionMode, "Set Any Selection Mode", "Sets the selection mode so that any element type may be selected.", EUserInterfaceActionType::None, FInputChord(EKeys::Four));
-	UI_COMMAND(QuadrangulateMesh, "Quadrangulate Mesh", "Quadrangulates the selected mesh.", EUserInterfaceActionType::Button, FInputChord());
+
+	for( TObjectIterator<UMeshEditorCommonCommand> CommonCommandCDOIter( RF_NoFlags ); CommonCommandCDOIter; ++CommonCommandCDOIter )
+	{
+		UMeshEditorCommonCommand* CommonCommandCDO = *CommonCommandCDOIter;
+		if( !( CommonCommandCDO->GetClass()->GetClassFlags() & CLASS_Abstract ) )
+		{
+			CommonCommandCDO->RegisterUICommand( this );
+		}
+	}
 }
 
 FMeshEditorAnyElementCommands::FMeshEditorAnyElementCommands() 
diff --git a/Engine/Plugins/Editor/MeshEditor/Source/MeshEditor/MeshEditorMode.cpp b/Engine/Plugins/Editor/MeshEditor/Source/MeshEditor/MeshEditorMode.cpp
--- a/Engine/Plugins/Editor/MeshEditor/Source/MeshEditor/MeshEditorMode.cpp
+++ b/Engine/Plugins/Editor/MeshEditor/Source/MeshEditor/MeshEditorMode.cpp
@@ -566,8 +566,6 @@ void FMeshEditorMode::BindCommands()
 	RegisterCommonCommand( MeshEditorCommonCommands.SetPolygonSelectionMode, FExecuteAction::CreateLambda( [this] { SetMeshElementSelectionMode( EEditableMeshElementType::Polygon ); } ) );
 	RegisterCommonCommand( MeshEditorCommonCommands.SetAnySelectionMode, FExecuteAction::CreateLambda( [this] { SetMeshElementSelectionMode( EEditableMeshElementType::Any ); } ) );
 
-	RegisterCommonCommand( MeshEditorCommonCommands.QuadrangulateMesh, FExecuteAction::CreateLambda( [ this ] { QuadrangulateMesh(); } ) );
-
 	// Register element-specific commands
 	RegisterVertexCommand( MeshEditorVertexCommands.WeldVertices, FExecuteAction::CreateLambda( [this] { WeldSelectedVertices(); } ) );
 
@@ -582,22 +580,33 @@ void FMeshEditorMode::BindCommands()
 		UMeshEditorCommand* CommandCDO = *CommandCDOIter;
 		if( !( CommandCDO->GetClass()->GetClassFlags() & CLASS_Abstract ) )
 		{
-			const EEditableMeshElementType ElementType = CommandCDO->GetElementType();
-			if( ElementType == EEditableMeshElementType::Vertex )
-			{
-				VertexActions.Emplace( CommandCDO->GetUICommandInfo(), CommandCDO->MakeUIAction( *this ) );
-			}
-			else if( ElementType == EEditableMeshElementType::Edge )
+			switch( CommandCDO->GetElementType() )
 			{
-				EdgeActions.Emplace( CommandCDO->GetUICommandInfo(), CommandCDO->MakeUIAction( *this ));
-			}
-			else if( ElementType == EEditableMeshElementType::Polygon )
-			{
-				PolygonActions.Emplace( CommandCDO->GetUICommandInfo(), CommandCDO->MakeUIAction( *this ) );
-			}
-			else
-			{
-				check( 0 );
+				case EEditableMeshElementType::Invalid:
+					{
+						// Common action
+						FUIAction UIAction = CommandCDO->MakeUIAction( *this );
+						CommonActions.Emplace( CommandCDO->GetUICommandInfo(), UIAction );
+						VertexActions.Emplace( CommandCDO->GetUICommandInfo(), UIAction );
+						EdgeActions.Emplace( CommandCDO->GetUICommandInfo(), UIAction );
+						PolygonActions.Emplace( CommandCDO->GetUICommandInfo(), UIAction );
+					}
+					break;
+
+				case EEditableMeshElementType::Vertex:
+					VertexActions.Emplace( CommandCDO->GetUICommandInfo(), CommandCDO->MakeUIAction( *this ) );
+					break;
+				
+				case EEditableMeshElementType::Edge:
+					EdgeActions.Emplace( CommandCDO->GetUICommandInfo(), CommandCDO->MakeUIAction( *this ) );
+					break;
+
+				case EEditableMeshElementType::Polygon:
+					PolygonActions.Emplace( CommandCDO->GetUICommandInfo(), CommandCDO->MakeUIAction( *this ) );
+					break;
+
+				default:
+					check( 0 );
 			}
 		}
 	}
@@ -673,18 +682,20 @@ void FMeshEditorMode::RegisterPolygonEditingMode( const TSharedPtr<FUICommandInf
 
 void FMeshEditorMode::RegisterCommonCommand( const TSharedPtr<FUICommandInfo>& Command, const FExecuteAction& ExecuteAction )
 {
-	CommonActions.Emplace( Command, FUIAction( ExecuteAction, FCanExecuteAction::CreateLambda( [this] { return GetSelectedEditableMeshes().Num() > 0; } ) ) );
-	VertexActions.Emplace( Command, FUIAction( ExecuteAction, FCanExecuteAction::CreateLambda( [this] { return GetSelectedEditableMeshes().Num() > 0; } ) ) );
-	EdgeActions.Emplace( Command, FUIAction( ExecuteAction, FCanExecuteAction::CreateLambda( [this] { return GetSelectedEditableMeshes().Num() > 0; } ) ) );
-	PolygonActions.Emplace( Command, FUIAction( ExecuteAction, FCanExecuteAction::CreateLambda( [this] { return GetSelectedEditableMeshes().Num() > 0; } ) ) );
+	FCanExecuteAction CanExecute = FCanExecuteAction::CreateLambda( [this] { return GetSelectedEditableMeshes().Num() > 0; } );
+	CommonActions.Emplace( Command, FUIAction( ExecuteAction, CanExecute ) );
+	VertexActions.Emplace( Command, FUIAction( ExecuteAction, CanExecute ) );
+	EdgeActions.Emplace( Command, FUIAction( ExecuteAction, CanExecute ) );
+	PolygonActions.Emplace( Command, FUIAction( ExecuteAction, CanExecute ) );
 }
 
 
 void FMeshEditorMode::RegisterAnyElementCommand( const TSharedPtr<FUICommandInfo>& Command, const FExecuteAction& ExecuteAction )
 {
-	VertexActions.Emplace( Command, FUIAction( ExecuteAction, FCanExecuteAction::CreateLambda( [this] { return GetSelectedMeshElementType() != EEditableMeshElementType::Invalid;} ) ) );
-	EdgeActions.Emplace( Command, FUIAction( ExecuteAction, FCanExecuteAction::CreateLambda( [this] { return GetSelectedMeshElementType() != EEditableMeshElementType::Invalid;} ) ) );
-	PolygonActions.Emplace( Command, FUIAction( ExecuteAction, FCanExecuteAction::CreateLambda( [this] { return GetSelectedMeshElementType() != EEditableMeshElementType::Invalid;} ) ) );
+	FCanExecuteAction CanExecute = FCanExecuteAction::CreateLambda( [this] { return GetSelectedMeshElementType() != EEditableMeshElementType::Invalid;} );
+	VertexActions.Emplace( Command, FUIAction( ExecuteAction, CanExecute ) );
+	EdgeActions.Emplace( Command, FUIAction( ExecuteAction, CanExecute ) );
+	PolygonActions.Emplace( Command, FUIAction( ExecuteAction, CanExecute ) );
 }
 
 
@@ -1372,36 +1383,6 @@ void FMeshEditorMode::AddOrRemoveSubdivisionLevel( const bool bShouldAdd )
 }
 
 
-void FMeshEditorMode::QuadrangulateMesh()
-{
-	if( ActiveAction == NAME_None )
-	{
-		if( GetSelectedEditableMeshes().Num() == 0 )
-		{
-			return;
-		}
-
-		FScopedTransaction Transaction( LOCTEXT( "UndoQuadrangulateMesh", "Quadrangulate Mesh" ) );
-
-		CommitSelectedMeshes();
-
-		const TArray<UEditableMesh*>& SelectedMeshes = GetSelectedEditableMeshes();
-
-		DeselectAllMeshElements();
-
-		for( UEditableMesh* EditableMesh : SelectedMeshes )
-		{
-			static TArray<FPolygonRef> NewPolygonRefs;
-			EditableMesh->StartModification( EMeshModificationType::Final, EMeshTopologyChange::TopologyChange );
-			EditableMesh->QuadrangulateMesh( NewPolygonRefs );
-			EditableMesh->EndModification();
-
-			TrackUndo( EditableMesh, EditableMesh->MakeUndo() );
-		}
-	}
-}
-
-
 void FMeshEditorMode::FrameSelectedElements( FEditorViewportClient* ViewportClient )
 {
 	FBox BoundingBox( ForceInitToZero );
@@ -4056,7 +4037,9 @@ void FMeshEditorMode::OnViewportInteractionInputAction( FEditorViewportClient& V
 										break;
 								}
 
-								if( CommandCDO->GetElementType() == SelectedMeshElementType && EquippedAction == CommandCDO->GetCommandName() )
+								const EEditableMeshElementType CommandElementType = CommandCDO->GetElementType();
+								if( ( CommandElementType == SelectedMeshElementType || CommandElementType == EEditableMeshElementType::Invalid || CommandElementType == EEditableMeshElementType::Any ) && 
+									EquippedAction == CommandCDO->GetCommandName() )
 								{
 									if( CommandCDO->TryStartingToDrag( *this, ViewportInteractor ) )
 									{
diff --git a/Engine/Plugins/Editor/MeshEditor/Source/MeshEditor/MeshEditorMode.h b/Engine/Plugins/Editor/MeshEditor/Source/MeshEditor/MeshEditorMode.h
--- a/Engine/Plugins/Editor/MeshEditor/Source/MeshEditor/MeshEditorMode.h
+++ b/Engine/Plugins/Editor/MeshEditor/Source/MeshEditor/MeshEditorMode.h
@@ -158,8 +158,6 @@ protected:
 	virtual EEditableMeshElementType GetSelectedMeshElementType() const override;
 	virtual bool IsMeshElementTypeSelected( EEditableMeshElementType ElementType ) const override { return GetSelectedMeshElementType() == ElementType; }
 	virtual bool IsMeshElementTypeSelectedOrIsActiveSelectionMode( EEditableMeshElementType ElementType ) const override { return GetMeshElementSelectionMode() == ElementType || GetSelectedMeshElementType() == ElementType; }
-	virtual const TArray<UEditableMesh*>& GetSelectedEditableMeshes() const override { return SelectedEditableMeshes; }
-	virtual const TArray<UEditableMesh*>& GetSelectedEditableMeshes() override { return SelectedEditableMeshes; }
 	virtual const TArray<TTuple<TSharedPtr<FUICommandInfo>, FUIAction>>& GetCommonActions() const override { return CommonActions; }
 	virtual const TArray<TTuple<TSharedPtr<FUICommandInfo>, FUIAction>>& GetVertexActions() const override { return VertexActions; }
 	virtual const TArray<TTuple<TSharedPtr<FUICommandInfo>, FUIAction>>& GetEdgeActions() const override { return EdgeActions; }
@@ -193,6 +191,8 @@ protected:
 	{
 		GetSelectedMeshesAndElements( EEditableMeshElementType::Polygon, /* Out */ OutMeshesAndPolygons );
 	}
+	virtual const TArray<UEditableMesh*>& GetSelectedEditableMeshes() const override { return SelectedEditableMeshes; }
+	virtual const TArray<UEditableMesh*>& GetSelectedEditableMeshes() override { return SelectedEditableMeshes; }
 	virtual void SelectMeshElements( const TArray<FMeshElement>& MeshElementsToSelect ) override;
 	virtual void DeselectAllMeshElements() override;
 	virtual void DeselectMeshElements( const TArray<FMeshElement>& MeshElementsToDeselect ) override;
@@ -295,9 +295,6 @@ protected:
 	/** Adds or removes a subdivision level for selected meshes */
 	void AddOrRemoveSubdivisionLevel( const bool bShouldAdd );
 
-	/** Quadrangulates the currently selected mesh */
-	void QuadrangulateMesh();
-
 	/** Moves the viewport camera to frame the currently selected elements */
 	void FrameSelectedElements( FEditorViewportClient* ViewportClient );
 
diff --git a/Engine/Plugins/Editor/MeshEditor/Source/MeshEditor/Public/IMeshEditorModeEditingContract.h b/Engine/Plugins/Editor/MeshEditor/Source/MeshEditor/Public/IMeshEditorModeEditingContract.h
--- a/Engine/Plugins/Editor/MeshEditor/Source/MeshEditor/Public/IMeshEditorModeEditingContract.h
+++ b/Engine/Plugins/Editor/MeshEditor/Source/MeshEditor/Public/IMeshEditorModeEditingContract.h
@@ -4,6 +4,7 @@
 
 #include "EditableMeshTypes.h"
 #include "MeshElement.h"
+#include "Change.h"
 
 
 class IMeshEditorModeEditingContract
@@ -32,6 +33,12 @@ public:
 	virtual void GetSelectedMeshesAndEdges( TMap<class UEditableMesh*, TArray<FMeshElement>>& OutMeshesAndEdges ) = 0;
 	virtual void GetSelectedMeshesAndPolygons( TMap<class UEditableMesh*, TArray<FMeshElement>>& OutMeshesAndPolygons ) = 0;
 
+	/** Gets all of the editable meshes that are currently selected (const) */
+	virtual const TArray<class UEditableMesh*>& GetSelectedEditableMeshes() const = 0;
+
+	/** Gets all of the editable meshes that are currently selected */
+	virtual const TArray<class UEditableMesh*>& GetSelectedEditableMeshes() = 0;
+
 	/** Selects the specified mesh elements */
 	virtual void SelectMeshElements( const TArray<FMeshElement>& MeshElementsToSelect ) = 0;
 
diff --git a/Engine/Plugins/Editor/MeshEditor/Source/MeshEditor/Public/IMeshEditorModeUIContract.h b/Engine/Plugins/Editor/MeshEditor/Source/MeshEditor/Public/IMeshEditorModeUIContract.h
--- a/Engine/Plugins/Editor/MeshEditor/Source/MeshEditor/Public/IMeshEditorModeUIContract.h
+++ b/Engine/Plugins/Editor/MeshEditor/Source/MeshEditor/Public/IMeshEditorModeUIContract.h
@@ -28,12 +28,6 @@ public:
 	/** Returns whether either the specified element type is selected, or we're in the selection mode for that element type */
 	virtual bool IsMeshElementTypeSelectedOrIsActiveSelectionMode( EEditableMeshElementType ElementType ) const = 0;
 
-	/** Gets all of the editable meshes that are currently selected (const) */
-	virtual const TArray<class UEditableMesh*>& GetSelectedEditableMeshes() const = 0;
-
-	/** Gets all of the editable meshes that are currently selected */
-	virtual const TArray<class UEditableMesh*>& GetSelectedEditableMeshes() = 0;
-
 	virtual const TArray<TTuple<TSharedPtr<FUICommandInfo>, FUIAction>>& GetCommonActions() const = 0;
 	virtual const TArray<TTuple<TSharedPtr<FUICommandInfo>, FUIAction>>& GetVertexActions() const = 0;
 	virtual const TArray<TTuple<TSharedPtr<FUICommandInfo>, FUIAction>>& GetEdgeActions() const = 0;
diff --git a/Engine/Plugins/Editor/MeshEditor/Source/MeshEditor/Public/MeshEditorCommands.h b/Engine/Plugins/Editor/MeshEditor/Source/MeshEditor/Public/MeshEditorCommands.h
--- a/Engine/Plugins/Editor/MeshEditor/Source/MeshEditor/Public/MeshEditorCommands.h
+++ b/Engine/Plugins/Editor/MeshEditor/Source/MeshEditor/Public/MeshEditorCommands.h
@@ -4,6 +4,8 @@
 
 #include "MeshElement.h"
 #include "UIAction.h"
+#include "UICommandInfo.h"
+#include "Commands.h"
 #include "MeshEditorCommands.generated.h"
 
 
@@ -37,7 +39,7 @@ public:
 	}
 
 	/** Allows this command to directly add a button to the VR Mode's radial menu */
-	virtual void AddToVRRadialMenuActionsMenu( class IMeshEditorModeUIContract& MeshEditorMode, class FMenuBuilder& MenuBuilder, TSharedPtr<FUICommandList> CommandList, const FName TEMPHACK_StyleSetName, class UVREditorMode* VRMode )
+	virtual void AddToVRRadialMenuActionsMenu( class IMeshEditorModeUIContract& MeshEditorMode, class FMenuBuilder& MenuBuilder, TSharedPtr<class FUICommandList> CommandList, const FName TEMPHACK_StyleSetName, class UVREditorMode* VRMode )
 	{
 	}
 
@@ -74,7 +76,7 @@ public:
 	}
 
 	/** Gets the UI command info for this command */
-	const TSharedPtr<class FUICommandInfo>& GetUICommandInfo() const
+	const TSharedPtr<FUICommandInfo>& GetUICommandInfo() const
 	{
 		return UICommandInfo;
 	}
@@ -107,6 +109,22 @@ protected:
 };
 
 
+UCLASS( abstract )
+class MESHEDITOR_API UMeshEditorCommonCommand : public UMeshEditorCommand
+{
+	GENERATED_BODY()
+
+public:
+
+	// Overrides
+	virtual EEditableMeshElementType GetElementType() const override
+	{
+		return EEditableMeshElementType::Invalid;
+	}
+
+};
+
+
 UCLASS( abstract )
 class MESHEDITOR_API UMeshEditorVertexCommand : public UMeshEditorCommand
 {
@@ -193,9 +211,6 @@ public:
 	TSharedPtr<FUICommandInfo> SetEdgeSelectionMode;
 	TSharedPtr<FUICommandInfo> SetPolygonSelectionMode;
 	TSharedPtr<FUICommandInfo> SetAnySelectionMode;
-
-	/** Quadrangulate mesh */
-	TSharedPtr<FUICommandInfo> QuadrangulateMesh;
 };
 
 
diff --git a/Engine/Plugins/Editor/MeshEditor/Source/MeshEditor/Public/MeshElement.h b/Engine/Plugins/Editor/MeshEditor/Source/MeshEditor/Public/MeshElement.h
--- a/Engine/Plugins/Editor/MeshEditor/Source/MeshEditor/Public/MeshElement.h
+++ b/Engine/Plugins/Editor/MeshEditor/Source/MeshEditor/Public/MeshElement.h
@@ -3,6 +3,7 @@
 #pragma once
 
 #include "EditableMeshTypes.h"
+#include "Components/PrimitiveComponent.h"
 
 /**
  * Uniquely identifies a specific element within a mesh
