Mesh Editor extensibility (phase 8) (part 2)
- Added missing source files

diff --git a/Engine/Plugins/Editor/MeshEditor/Source/PolygonModeling/QuadrangulateMesh.cpp b/Engine/Plugins/Editor/MeshEditor/Source/PolygonModeling/QuadrangulateMesh.cpp
new file mode 100644
--- /dev/null
+++ b/Engine/Plugins/Editor/MeshEditor/Source/PolygonModeling/QuadrangulateMesh.cpp
@@ -0,0 +1,54 @@
+// Copyright 1998-2016 Epic Games, Inc. All Rights Reserved.
+
+#include "QuadrangulateMesh.h"
+#include "IMeshEditorModeEditingContract.h"
+#include "IMeshEditorModeUIContract.h"
+#include "UICommandInfo.h"
+#include "EditableMesh.h"
+#include "MeshElement.h"
+#include "ScopedTransaction.h"
+
+#define LOCTEXT_NAMESPACE "MeshEditorMode"
+
+
+void UQuadrangulateMeshCommand::RegisterUICommand( FBindingContext* BindingContext )
+{
+	UI_COMMAND_EXT( BindingContext, /* Out */ UICommandInfo, "QuadrangulateMesh", "Quadrangulate Mesh", "Quadrangulates the selected mesh.", EUserInterfaceActionType::Button, FInputChord() );
+}
+
+
+void UQuadrangulateMeshCommand::Execute( IMeshEditorModeEditingContract& MeshEditorMode )
+{
+	if( MeshEditorMode.GetActiveAction() != NAME_None )
+	{
+		return;
+	}
+
+	if( MeshEditorMode.GetSelectedEditableMeshes().Num() == 0 )
+	{
+		return;
+	}
+
+	FScopedTransaction Transaction( LOCTEXT( "UndoQuadrangulateMesh", "Quadrangulate Mesh" ) );
+
+	MeshEditorMode.CommitSelectedMeshes();
+
+	const TArray<UEditableMesh*>& SelectedMeshes = MeshEditorMode.GetSelectedEditableMeshes();
+
+	MeshEditorMode.DeselectAllMeshElements();
+
+	for( UEditableMesh* EditableMesh : SelectedMeshes )
+	{
+		static TArray<FPolygonRef> NewPolygonRefs;
+		EditableMesh->StartModification( EMeshModificationType::Final, EMeshTopologyChange::TopologyChange );
+		EditableMesh->QuadrangulateMesh( NewPolygonRefs );
+		EditableMesh->EndModification();
+
+		MeshEditorMode.TrackUndo( EditableMesh, EditableMesh->MakeUndo() );
+	}
+}
+
+
+
+#undef LOCTEXT_NAMESPACE
+
diff --git a/Engine/Plugins/Editor/MeshEditor/Source/PolygonModeling/QuadrangulateMesh.h b/Engine/Plugins/Editor/MeshEditor/Source/PolygonModeling/QuadrangulateMesh.h
new file mode 100644
--- /dev/null
+++ b/Engine/Plugins/Editor/MeshEditor/Source/PolygonModeling/QuadrangulateMesh.h
@@ -0,0 +1,21 @@
+// Copyright 1998-2016 Epic Games, Inc. All Rights Reserved.
+
+#pragma once
+
+#include "MeshEditorCommands.h"
+#include "QuadrangulateMesh.generated.h"
+
+
+/** Quadrangulates the currently selected mesh */
+UCLASS()
+class UQuadrangulateMeshCommand : public UMeshEditorCommonCommand
+{
+	GENERATED_BODY()
+
+protected:
+
+	// Overrides
+	virtual void RegisterUICommand( class FBindingContext* BindingContext ) override;
+	virtual void Execute( class IMeshEditorModeEditingContract& MeshEditorMode ) override;
+
+};
