Mesh Editor: Fixed crash when selecting mesh elements in Simulate Mode
- Transactions aren't available in Simulate, but we were trying to store undo state

diff --git a/Engine/Plugins/Editor/MeshEditor/Source/MeshEditor/MeshEditorMode.cpp b/Engine/Plugins/Editor/MeshEditor/Source/MeshEditor/MeshEditorMode.cpp
--- a/Engine/Plugins/Editor/MeshEditor/Source/MeshEditor/MeshEditorMode.cpp
+++ b/Engine/Plugins/Editor/MeshEditor/Source/MeshEditor/MeshEditorMode.cpp
@@ -4982,8 +4982,10 @@ void FMeshEditorMode::TrackUndo( UObject* Object, TUniquePtr<FChange> RevertChan
 			// Did you forget to use an FScopedTransaction?  If GUndo was null, then most likely we forgot to wrap this call within an editor transaction.
 			// The only exception is in Simulate mode, where Undo is not allowed.
 			check( GUndo != nullptr || GEditor == nullptr || GEditor->bIsSimulatingInEditor );
-
-			GUndo->StoreUndo( Object, MoveTemp( RevertChange ) );
+			if( GUndo != nullptr )
+			{
+				GUndo->StoreUndo( Object, MoveTemp( RevertChange ) );
+			}
 		}
 		else
 		{
