Mesh Editor: Minor refactor to improve extensibility of mesh commands
- Removed type-specific base classes for mesh editor commands
- We now have a different base class for 'instant' commands vs 'edit' commands instead
- Removed 'F1' shortcut for move mode.  Will be dealing with that soon.

diff --git a/Engine/Plugins/Editor/MeshEditor/Source/MeshEditor/MeshEditorCommands.cpp b/Engine/Plugins/Editor/MeshEditor/Source/MeshEditor/MeshEditorCommands.cpp
--- a/Engine/Plugins/Editor/MeshEditor/Source/MeshEditor/MeshEditorCommands.cpp
+++ b/Engine/Plugins/Editor/MeshEditor/Source/MeshEditor/MeshEditorCommands.cpp
@@ -10,22 +10,11 @@
 #define LOCTEXT_NAMESPACE "MeshEditorCommands"
 
 
-FUIAction UMeshEditorCommand::MakeUIAction( IMeshEditorModeUIContract& MeshEditorMode )
+FUIAction UMeshEditorInstantCommand::MakeUIAction( IMeshEditorModeUIContract& MeshEditorMode )
 {
 	const EEditableMeshElementType ElementType = GetElementType();
 
 	FUIAction UIAction;
-	if( IsMode() )
-	{
-		const FName CommandName = GetCommandName();
-
-		UIAction = FUIAction(
-			FExecuteAction::CreateLambda( [&MeshEditorMode, ElementType, CommandName] { MeshEditorMode.SetEquippedAction( ElementType, CommandName ); } ),
-			FCanExecuteAction::CreateLambda( [&MeshEditorMode, ElementType] { return MeshEditorMode.IsMeshElementTypeSelectedOrIsActiveSelectionMode( ElementType ); } ),
-			FIsActionChecked::CreateLambda( [&MeshEditorMode, ElementType, CommandName] { return ( MeshEditorMode.GetEquippedAction( ElementType ) == CommandName ); } )
-		);
-	}
-	else
 	{
 		FExecuteAction ExecuteAction( FExecuteAction::CreateLambda( [&MeshEditorMode, this]
 		{
@@ -53,10 +42,29 @@ FUIAction UMeshEditorCommand::MakeUIAction( IMeshEditorModeUIContract& MeshEdito
 }
 
 
+FUIAction UMeshEditorEditCommand::MakeUIAction( IMeshEditorModeUIContract& MeshEditorMode )
+{
+	const EEditableMeshElementType ElementType = GetElementType();
+
+	FUIAction UIAction;
+	{
+		const FName CommandName = GetCommandName();
+
+		UIAction = FUIAction(
+			FExecuteAction::CreateLambda( [&MeshEditorMode, ElementType, CommandName] { MeshEditorMode.SetEquippedAction( ElementType, CommandName ); } ),
+			FCanExecuteAction::CreateLambda( [&MeshEditorMode, ElementType] { return MeshEditorMode.IsMeshElementTypeSelectedOrIsActiveSelectionMode( ElementType ); } ),
+			FIsActionChecked::CreateLambda( [&MeshEditorMode, ElementType, CommandName] { return ( MeshEditorMode.GetEquippedAction( ElementType ) == CommandName ); } )
+		);
+	}
+
+	return UIAction;
+}
+
+
 FMeshEditorCommonCommands::FMeshEditorCommonCommands() 
 	: TCommands<FMeshEditorCommonCommands>(
 		"MeshEditorCommon",
-		LOCTEXT( "MeshEditorGeneral", "Mesh Editor Common" ),
+		LOCTEXT( "MeshEditorCommon", "Mesh Editor Common" ),
 		"MainFrame",
 		FMeshEditorStyle::GetStyleSetName()	)
 {
@@ -78,12 +86,15 @@ void FMeshEditorCommonCommands::RegisterCommands()
 	UI_COMMAND(SetPolygonSelectionMode, "Set Polygon Selection Mode", "Sets the selection mode so that only polygons will be selected.", EUserInterfaceActionType::None, FInputChord(EKeys::Three));
 	UI_COMMAND(SetAnySelectionMode, "Set Any Selection Mode", "Sets the selection mode so that any element type may be selected.", EUserInterfaceActionType::None, FInputChord(EKeys::Four));
 
-	for( TObjectIterator<UMeshEditorCommonCommand> CommonCommandCDOIter( RF_NoFlags ); CommonCommandCDOIter; ++CommonCommandCDOIter )
+	for( TObjectIterator<UMeshEditorCommand> CommandCDOIter( RF_NoFlags ); CommandCDOIter; ++CommandCDOIter )
 	{
-		UMeshEditorCommonCommand* CommonCommandCDO = *CommonCommandCDOIter;
-		if( !( CommonCommandCDO->GetClass()->GetClassFlags() & CLASS_Abstract ) )
+		UMeshEditorCommand* CommandCDO = *CommandCDOIter;
+		if( !( CommandCDO->GetClass()->GetClassFlags() & CLASS_Abstract ) )
 		{
-			CommonCommandCDO->RegisterUICommand( this );
+			if( CommandCDO->GetElementType() == EEditableMeshElementType::Invalid )
+			{
+				CommandCDO->RegisterUICommand( this );
+			}
 		}
 	}
 }
@@ -115,17 +126,19 @@ FMeshEditorVertexCommands::FMeshEditorVertexCommands()
 
 void FMeshEditorVertexCommands::RegisterCommands()
 {
-	UI_COMMAND(MoveVertex, "Move Vertex Mode", "Set the primary action to move vertices.", EUserInterfaceActionType::RadioButton, FInputChord(EKeys::F1));
+	UI_COMMAND(MoveVertex, "Move Vertex Mode", "Set the primary action to move vertices.", EUserInterfaceActionType::RadioButton, FInputChord());
 
 	UI_COMMAND(WeldVertices, "Weld Vertices", "Weld the selected vertices, keeping the first selected vertex.", EUserInterfaceActionType::Button, FInputChord());
 
-	// @todo mesheditor extensibility: What's the plan for default keybinds for commands registered in a modular way?  Should we suggest available keys?
-	for( TObjectIterator<UMeshEditorVertexCommand> VertexCommandCDOIter( RF_NoFlags ); VertexCommandCDOIter; ++VertexCommandCDOIter )
+	for( TObjectIterator<UMeshEditorCommand> CommandCDOIter( RF_NoFlags ); CommandCDOIter; ++CommandCDOIter )
 	{
-		UMeshEditorVertexCommand* VertexCommandCDO = *VertexCommandCDOIter;
-		if( !( VertexCommandCDO->GetClass()->GetClassFlags() & CLASS_Abstract ) )
+		UMeshEditorCommand* CommandCDO = *CommandCDOIter;
+		if( !( CommandCDO->GetClass()->GetClassFlags() & CLASS_Abstract ) )
 		{
-			VertexCommandCDO->RegisterUICommand( this );
+			if( CommandCDO->GetElementType() == EEditableMeshElementType::Vertex )
+			{
+				CommandCDO->RegisterUICommand( this );
+			}
 		}
 	}
 }
@@ -142,16 +155,19 @@ FMeshEditorEdgeCommands::FMeshEditorEdgeCommands()
 
 void FMeshEditorEdgeCommands::RegisterCommands()
 {
-	UI_COMMAND(MoveEdge, "Move Edge Mode", "Set the primary action to move edges.", EUserInterfaceActionType::RadioButton, FInputChord(EKeys::F1));
+	UI_COMMAND(MoveEdge, "Move Edge Mode", "Set the primary action to move edges.", EUserInterfaceActionType::RadioButton, FInputChord());
 
 	UI_COMMAND(SelectEdgeLoop, "Select Edge Loop", "Select the edge loops which contain the selected edges.", EUserInterfaceActionType::Button, FInputChord(EKeys::Two, EModifierKey::Shift));
 
-	for( TObjectIterator<UMeshEditorEdgeCommand> EdgeCommandCDOIter( RF_NoFlags ); EdgeCommandCDOIter; ++EdgeCommandCDOIter )
+	for( TObjectIterator<UMeshEditorCommand> CommandCDOIter( RF_NoFlags ); CommandCDOIter; ++CommandCDOIter )
 	{
-		UMeshEditorEdgeCommand* EdgeCommandCDO = *EdgeCommandCDOIter;
-		if( !( EdgeCommandCDO->GetClass()->GetClassFlags() & CLASS_Abstract ) )
+		UMeshEditorCommand* CommandCDO = *CommandCDOIter;
+		if( !( CommandCDO->GetClass()->GetClassFlags() & CLASS_Abstract ) )
 		{
-			EdgeCommandCDO->RegisterUICommand( this );
+			if( CommandCDO->GetElementType() == EEditableMeshElementType::Edge )
+			{
+				CommandCDO->RegisterUICommand( this );
+			}
 		}
 	}
 }
@@ -168,17 +184,20 @@ FMeshEditorPolygonCommands::FMeshEditorPolygonCommands()
 
 void FMeshEditorPolygonCommands::RegisterCommands()
 {
-	UI_COMMAND(MovePolygon, "Move Polygon Mode", "Set the primary action to move polygons.", EUserInterfaceActionType::RadioButton, FInputChord(EKeys::F1));
+	UI_COMMAND(MovePolygon, "Move Polygon Mode", "Set the primary action to move polygons.", EUserInterfaceActionType::RadioButton, FInputChord());
 	UI_COMMAND(FlipPolygon, "Flip Polygon", "Flip the currently selected polygons.", EUserInterfaceActionType::Button, FInputChord(EKeys::F, EModifierKey::Shift));
 	UI_COMMAND(TriangulatePolygon, "Triangulate Polygon", "Triangulate the currently selected polygons.", EUserInterfaceActionType::Button, FInputChord(EKeys::T));
 	UI_COMMAND(AssignMaterial, "Assign Material", "Assigns the highlighted material in the Content Browser to the currently selected polygons.", EUserInterfaceActionType::Button, FInputChord(EKeys::M));
 
-	for( TObjectIterator<UMeshEditorPolygonCommand> PolygonCommandCDOIter( RF_NoFlags ); PolygonCommandCDOIter; ++PolygonCommandCDOIter )
+	for( TObjectIterator<UMeshEditorCommand> CommandCDOIter( RF_NoFlags ); CommandCDOIter; ++CommandCDOIter )
 	{
-		UMeshEditorPolygonCommand* PolygonCommandCDO = *PolygonCommandCDOIter;
-		if( !( PolygonCommandCDO->GetClass()->GetClassFlags() & CLASS_Abstract ) )
+		UMeshEditorCommand* CommandCDO = *CommandCDOIter;
+		if( !( CommandCDO->GetClass()->GetClassFlags() & CLASS_Abstract ) )
 		{
-			PolygonCommandCDO->RegisterUICommand( this );
+			if( CommandCDO->GetElementType() == EEditableMeshElementType::Polygon )
+			{
+				CommandCDO->RegisterUICommand( this );
+			}
 		}
 	}
 }
diff --git a/Engine/Plugins/Editor/MeshEditor/Source/MeshEditor/MeshEditorMode.cpp b/Engine/Plugins/Editor/MeshEditor/Source/MeshEditor/MeshEditorMode.cpp
--- a/Engine/Plugins/Editor/MeshEditor/Source/MeshEditor/MeshEditorMode.cpp
+++ b/Engine/Plugins/Editor/MeshEditor/Source/MeshEditor/MeshEditorMode.cpp
@@ -3249,9 +3249,9 @@ void FMeshEditorMode::UpdateActiveAction( const bool bIsActionFinishing )
 	{
 		// Check for registered commands that are active right now
 		bool bFoundValidCommand = false;
-		for( TObjectIterator<UMeshEditorCommand> CommandCDOIter( RF_NoFlags ); CommandCDOIter; ++CommandCDOIter )
+		for( TObjectIterator<UMeshEditorEditCommand> CommandCDOIter( RF_NoFlags ); CommandCDOIter; ++CommandCDOIter )
 		{
-			UMeshEditorCommand* CommandCDO = *CommandCDOIter;
+			UMeshEditorEditCommand* CommandCDO = *CommandCDOIter;
 			if( !( CommandCDO->GetClass()->GetClassFlags() & CLASS_Abstract ) )
 			{
 				if( ActiveAction == CommandCDO->GetCommandName() )
@@ -4016,9 +4016,9 @@ void FMeshEditorMode::OnViewportInteractionInputAction( FEditorViewportClient& V
 					}
 					else
 					{
-						for( TObjectIterator<UMeshEditorCommand> CommandCDOIter( RF_NoFlags ); CommandCDOIter; ++CommandCDOIter )
+						for( TObjectIterator<UMeshEditorEditCommand> CommandCDOIter( RF_NoFlags ); CommandCDOIter; ++CommandCDOIter )
 						{
-							UMeshEditorCommand* CommandCDO = *CommandCDOIter;
+							UMeshEditorEditCommand* CommandCDO = *CommandCDOIter;
 							if( !( CommandCDO->GetClass()->GetClassFlags() & CLASS_Abstract ) )
 							{
 								FName EquippedAction = NAME_None;
diff --git a/Engine/Plugins/Editor/MeshEditor/Source/MeshEditor/Public/MeshEditorCommands.h b/Engine/Plugins/Editor/MeshEditor/Source/MeshEditor/Public/MeshEditorCommands.h
--- a/Engine/Plugins/Editor/MeshEditor/Source/MeshEditor/Public/MeshEditorCommands.h
+++ b/Engine/Plugins/Editor/MeshEditor/Source/MeshEditor/Public/MeshEditorCommands.h
@@ -16,28 +16,20 @@ class MESHEDITOR_API UMeshEditorCommand : public UObject
 
 public:
 
-	/** What type of mesh element is this command for? */
+	/** Which mesh element does this command apply to? */
 	virtual EEditableMeshElementType GetElementType() const PURE_VIRTUAL(,return EEditableMeshElementType::Invalid;);
 
 	/** Registers the UI command for this mesh editor command */
 	virtual void RegisterUICommand( class FBindingContext* BindingContext ) PURE_VIRTUAL(,);
 
+	/** Creates an UI action for this command */
+	virtual FUIAction MakeUIAction( class IMeshEditorModeUIContract& MeshEditorMode ) PURE_VIRTUAL(,return FUIAction(););
+
 	/** Runs this command */
 	virtual void Execute( class IMeshEditorModeEditingContract& MeshEditorMode )
 	{
 	}
 
-	/** Called when the user starts to drag on an element.  If this returns true, then the action will begin and ApplyDuringDrag() will be called each frame until the user releases the button. */
-	virtual bool TryStartingToDrag( IMeshEditorModeEditingContract& MeshEditorMode, class UViewportInteractor* ViewportInteractor )
-	{
-		return true;
-	}
-
-	/** Applies this command every frame while dragging */
-	virtual void ApplyDuringDrag( class IMeshEditorModeEditingContract& MeshEditorMode, class UViewportInteractor* ViewportInteractor )
-	{
-	}
-
 	/** Allows this command to directly add a button to the VR Mode's radial menu */
 	virtual void AddToVRRadialMenuActionsMenu( class IMeshEditorModeUIContract& MeshEditorMode, class FMenuBuilder& MenuBuilder, TSharedPtr<class FUICommandList> CommandList, const FName TEMPHACK_StyleSetName, class UVREditorMode* VRMode )
 	{
@@ -49,127 +41,90 @@ public:
 		return UICommandInfo->GetCommandName();
 	}
 
-	/** Gets the text to send to the transaction system when creating an undo/redo event for this action */
-	FText GetUndoText() const
-	{
-		check( !bIsMode || !UndoText.IsEmpty() );	// All mode-based commands must have supplied UndoText.  Instantaneous commands handle their own undo/redo.
-		return UndoText;
-	}
-
-	/** Returns true if this is a mesh editing 'Mode' that the user will stay in to perform the action multiple times, or false
-		if the action applies instantly */
-	bool IsMode() const
-	{
-		return bIsMode;
-	}
-
-	/** Returns whether we rely on a hover location under the interactor being updated as we drag during this action */
-	bool NeedsHoverLocation() const
-	{
-		return bNeedsHoverLocation;
-	}
-
-	/** Returns whether this command will kick off regular free translation of the selected mesh elements when dragging starts */
-	bool NeedsDraggingInitiated() const
-	{
-		return bNeedsDraggingInitiated;
-	}
-
 	/** Gets the UI command info for this command */
 	const TSharedPtr<FUICommandInfo>& GetUICommandInfo() const
 	{
 		return UICommandInfo;
 	}
 
-	/** Creates an UI action for this command */
-	FUIAction MakeUIAction( class IMeshEditorModeUIContract& MeshEditorMode );
-
 
 protected:
 
-	/** The text to send to the transaction system when creating an undo / redo event for this action */
-	UPROPERTY( EditAnywhere, Category=MeshEditor )
-	FText UndoText;
-
-	/** True if this is a mesh editing 'Mode' that the user will stay in to perform the action multiple times, or false
-	    if the action applies instantly */
-	UPROPERTY( EditAnywhere, Category=MeshEditor )
-	bool bIsMode;
-
-	/** Whether this command will kick off regular free translation of the selected mesh elements when dragging starts */
-	UPROPERTY( EditAnywhere, Category=MeshEditor )
-	bool bNeedsDraggingInitiated;
-
-	/** Whether we rely on a hover location under the interactor being updated as we drag during this action */
-	UPROPERTY( EditAnywhere, Category=MeshEditor )
-	bool bNeedsHoverLocation;
-
 	/** Our UI command for this action */
 	TSharedPtr<FUICommandInfo> UICommandInfo;
 };
 
 
 UCLASS( abstract )
-class MESHEDITOR_API UMeshEditorCommonCommand : public UMeshEditorCommand
+class MESHEDITOR_API UMeshEditorInstantCommand : public UMeshEditorCommand
 {
 	GENERATED_BODY()
 
 public:
 
-	// Overrides
-	virtual EEditableMeshElementType GetElementType() const override
-	{
-		return EEditableMeshElementType::Invalid;
-	}
+	// UMeshEditorCommand overrides
+	virtual FUIAction MakeUIAction( class IMeshEditorModeUIContract& MeshEditorMode ) override;
+
+protected:
 
 };
 
 
+
 UCLASS( abstract )
-class MESHEDITOR_API UMeshEditorVertexCommand : public UMeshEditorCommand
+class MESHEDITOR_API UMeshEditorEditCommand : public UMeshEditorCommand
 {
 	GENERATED_BODY()
 
 public:
 
-	// Overrides
-	virtual EEditableMeshElementType GetElementType() const override
+	// UMeshEditorCommand overrides
+	virtual FUIAction MakeUIAction( class IMeshEditorModeUIContract& MeshEditorMode ) override;
+
+	/** Called when the user starts to drag on an element.  If this returns true, then the action will begin and ApplyDuringDrag() will be called each frame until the user releases the button. */
+	virtual bool TryStartingToDrag( IMeshEditorModeEditingContract& MeshEditorMode, class UViewportInteractor* ViewportInteractor )
 	{
-		return EEditableMeshElementType::Vertex;
+		return true;
 	}
 
-};
-
-
-UCLASS( abstract )
-class MESHEDITOR_API UMeshEditorEdgeCommand : public UMeshEditorCommand
-{
-	GENERATED_BODY()
+	/** Applies this command every frame while dragging */
+	virtual void ApplyDuringDrag( class IMeshEditorModeEditingContract& MeshEditorMode, class UViewportInteractor* ViewportInteractor )
+	{
+	}
 
-public:
+	/** Gets the text to send to the transaction system when creating an undo/redo event for this action */
+	FText GetUndoText() const
+	{
+		check( !UndoText.IsEmpty() );
+		return UndoText;
+	}
 
-	// Overrides
-	virtual EEditableMeshElementType GetElementType() const override
+	/** Returns whether we rely on a hover location under the interactor being updated as we drag during this action */
+	bool NeedsHoverLocation() const
 	{
-		return EEditableMeshElementType::Edge;
+		return bNeedsHoverLocation;
 	}
 
-};
+	/** Returns whether this command will kick off regular free translation of the selected mesh elements when dragging starts */
+	bool NeedsDraggingInitiated() const
+	{
+		return bNeedsDraggingInitiated;
+	}
 
 
-UCLASS( abstract )
-class MESHEDITOR_API UMeshEditorPolygonCommand : public UMeshEditorCommand
-{
-	GENERATED_BODY()
+protected:
 
-public:
+	/** The text to send to the transaction system when creating an undo / redo event for this action */
+	UPROPERTY( EditAnywhere, Category=MeshEditor )
+	FText UndoText;
 
-	// Overrides
-	virtual EEditableMeshElementType GetElementType() const override
-	{
-		return EEditableMeshElementType::Polygon;
-	}
+	/** Whether this command will kick off regular free translation of the selected mesh elements when dragging starts */
+	UPROPERTY( EditAnywhere, Category=MeshEditor )
+	bool bNeedsDraggingInitiated;
 
+	/** Whether we rely on a hover location under the interactor being updated as we drag during this action */
+	UPROPERTY( EditAnywhere, Category=MeshEditor )
+	bool bNeedsHoverLocation;
 };
 
 
diff --git a/Engine/Plugins/Editor/MeshEditor/Source/PolygonModeling/BevelOrInsetPolygon.h b/Engine/Plugins/Editor/MeshEditor/Source/PolygonModeling/BevelOrInsetPolygon.h
--- a/Engine/Plugins/Editor/MeshEditor/Source/PolygonModeling/BevelOrInsetPolygon.h
+++ b/Engine/Plugins/Editor/MeshEditor/Source/PolygonModeling/BevelOrInsetPolygon.h
@@ -8,7 +8,7 @@
 
 /** Adds a beveled edge to an existing polygon */
 UCLASS()
-class UBevelPolygonCommand : public UMeshEditorPolygonCommand
+class UBevelPolygonCommand : public UMeshEditorEditCommand
 {
 	GENERATED_BODY()
 
@@ -16,13 +16,16 @@ protected:
 
 	UBevelPolygonCommand()
 	{
-		bIsMode = true;
 		UndoText = NSLOCTEXT( "MeshEditor", "UndoBevelPolygon", "Bevel Polygon" );
 		bNeedsHoverLocation = true;
 		bNeedsDraggingInitiated = false;
 	}
 
 	// Overrides
+	virtual EEditableMeshElementType GetElementType() const override
+	{
+		return EEditableMeshElementType::Polygon;
+	}
 	virtual void RegisterUICommand( class FBindingContext* BindingContext ) override;
 	virtual bool TryStartingToDrag( IMeshEditorModeEditingContract& MeshEditorMode, class UViewportInteractor* ViewportInteractor ) override;
 	virtual void ApplyDuringDrag( IMeshEditorModeEditingContract& MeshEditorMode, class UViewportInteractor* ViewportInteractor ) override;
@@ -33,7 +36,7 @@ protected:
 
 /** Adds a new polygon on the inside of an existing polygon, allowing the user to drag to set exactly where it should be placed. */
 UCLASS()
-class UInsetPolygonCommand : public UMeshEditorPolygonCommand
+class UInsetPolygonCommand : public UMeshEditorEditCommand
 {
 	GENERATED_BODY()
 
@@ -41,13 +44,16 @@ protected:
 
 	UInsetPolygonCommand()
 	{
-		bIsMode = true;
 		UndoText = NSLOCTEXT( "MeshEditor", "UndoInsetPolygon", "Inset Polygon" );
 		bNeedsHoverLocation = true;
 		bNeedsDraggingInitiated = false;
 	}
 
 	// Overrides
+	virtual EEditableMeshElementType GetElementType() const override
+	{
+		return EEditableMeshElementType::Polygon;
+	}
 	virtual void RegisterUICommand( class FBindingContext* BindingContext ) override;
 	virtual bool TryStartingToDrag( class IMeshEditorModeEditingContract& MeshEditorMode, class UViewportInteractor* ViewportInteractor ) override;
 	virtual void ApplyDuringDrag( IMeshEditorModeEditingContract& MeshEditorMode, class UViewportInteractor* ViewportInteractor ) override;
diff --git a/Engine/Plugins/Editor/MeshEditor/Source/PolygonModeling/EditVertexOrEdgeSharpness.h b/Engine/Plugins/Editor/MeshEditor/Source/PolygonModeling/EditVertexOrEdgeSharpness.h
--- a/Engine/Plugins/Editor/MeshEditor/Source/PolygonModeling/EditVertexOrEdgeSharpness.h
+++ b/Engine/Plugins/Editor/MeshEditor/Source/PolygonModeling/EditVertexOrEdgeSharpness.h
@@ -8,7 +8,7 @@
 
 /** For subdivision meshes, edits how sharp a vertex corner is by dragging in space */
 UCLASS()
-class UEditVertexCornerSharpnessCommand : public UMeshEditorVertexCommand
+class UEditVertexCornerSharpnessCommand : public UMeshEditorEditCommand
 {
 	GENERATED_BODY()
 
@@ -16,13 +16,16 @@ protected:
 
 	UEditVertexCornerSharpnessCommand()
 	{
-		bIsMode = true;
 		UndoText = NSLOCTEXT( "MeshEditor", "UndoEditVertexCornerSharpness", "Edit Vertex Corner Sharpness" );
 		bNeedsHoverLocation = false;
 		bNeedsDraggingInitiated = false;
 	}
 
 	// Overrides
+	virtual EEditableMeshElementType GetElementType() const override
+	{
+		return EEditableMeshElementType::Vertex;
+	}
 	virtual void RegisterUICommand( class FBindingContext* BindingContext ) override;
 	virtual bool TryStartingToDrag( IMeshEditorModeEditingContract& MeshEditorMode, class UViewportInteractor* ViewportInteractor ) override;
 	virtual void ApplyDuringDrag( IMeshEditorModeEditingContract& MeshEditorMode, class UViewportInteractor* ViewportInteractor ) override;
@@ -31,7 +34,7 @@ protected:
 
 /** For subdivision meshes, edits how sharp an edge crease is by dragging in space */
 UCLASS()
-class UEditEdgeCreaseSharpnessCommand : public UMeshEditorEdgeCommand
+class UEditEdgeCreaseSharpnessCommand : public UMeshEditorEditCommand
 {
 	GENERATED_BODY()
 
@@ -39,13 +42,16 @@ protected:
 
 	UEditEdgeCreaseSharpnessCommand()
 	{
-		bIsMode = true;
 		UndoText = NSLOCTEXT( "MeshEditor", "UndoEditEdgeCreaseSharpness", "Edit Edge Crease Sharpness" );
 		bNeedsHoverLocation = false;
 		bNeedsDraggingInitiated = false;
 	}
 
 	// Overrides
+	virtual EEditableMeshElementType GetElementType() const override
+	{
+		return EEditableMeshElementType::Edge;
+	}
 	virtual void RegisterUICommand( class FBindingContext* BindingContext ) override;
 	virtual bool TryStartingToDrag( IMeshEditorModeEditingContract& MeshEditorMode, class UViewportInteractor* ViewportInteractor ) override;
 	virtual void ApplyDuringDrag( IMeshEditorModeEditingContract& MeshEditorMode, class UViewportInteractor* ViewportInteractor ) override;
diff --git a/Engine/Plugins/Editor/MeshEditor/Source/PolygonModeling/ExtendEdge.h b/Engine/Plugins/Editor/MeshEditor/Source/PolygonModeling/ExtendEdge.h
--- a/Engine/Plugins/Editor/MeshEditor/Source/PolygonModeling/ExtendEdge.h
+++ b/Engine/Plugins/Editor/MeshEditor/Source/PolygonModeling/ExtendEdge.h
@@ -8,7 +8,7 @@
 
 /** Extends an edge by making a copy of it and allowing you to move it around */
 UCLASS()
-class UExtendEdgeCommand : public UMeshEditorEdgeCommand
+class UExtendEdgeCommand : public UMeshEditorEditCommand
 {
 	GENERATED_BODY()
 
@@ -16,13 +16,16 @@ protected:
 
 	UExtendEdgeCommand()
 	{
-		bIsMode = true;
 		UndoText = NSLOCTEXT( "MeshEditor", "UndoExtendEdge", "Extend Edge" );
 		bNeedsHoverLocation = false;
 		bNeedsDraggingInitiated = true;
 	}
 
 	// Overrides
+	virtual EEditableMeshElementType GetElementType() const override
+	{
+		return EEditableMeshElementType::Edge;
+	}
 	virtual void RegisterUICommand( class FBindingContext* BindingContext ) override;
 	virtual void ApplyDuringDrag( IMeshEditorModeEditingContract& MeshEditorMode, class UViewportInteractor* ViewportInteractor ) override;
 	virtual void AddToVRRadialMenuActionsMenu( class IMeshEditorModeUIContract& MeshEditorMode, class FMenuBuilder& MenuBuilder, TSharedPtr<FUICommandList> CommandList, const FName TEMPHACK_StyleSetName, class UVREditorMode* VRMode ) override;
diff --git a/Engine/Plugins/Editor/MeshEditor/Source/PolygonModeling/ExtendVertex.h b/Engine/Plugins/Editor/MeshEditor/Source/PolygonModeling/ExtendVertex.h
--- a/Engine/Plugins/Editor/MeshEditor/Source/PolygonModeling/ExtendVertex.h
+++ b/Engine/Plugins/Editor/MeshEditor/Source/PolygonModeling/ExtendVertex.h
@@ -8,7 +8,7 @@
 
 /** Extend a vertex by making a copy of it, creating new polygons to join the geometry together */
 UCLASS()
-class UExtendVertexCommand : public UMeshEditorVertexCommand
+class UExtendVertexCommand : public UMeshEditorEditCommand
 {
 	GENERATED_BODY()
 
@@ -16,13 +16,16 @@ protected:
 
 	UExtendVertexCommand()
 	{
-		bIsMode = true;
 		UndoText = NSLOCTEXT( "MeshEditor", "UndoExtendVertex", "Extend Vertex" );
 		bNeedsHoverLocation = false;
 		bNeedsDraggingInitiated = true;
 	}
 
 	// Overrides
+	virtual EEditableMeshElementType GetElementType() const override
+	{
+		return EEditableMeshElementType::Vertex;
+	}
 	virtual void RegisterUICommand( class FBindingContext* BindingContext ) override;
 	virtual void ApplyDuringDrag( IMeshEditorModeEditingContract& MeshEditorMode, class UViewportInteractor* ViewportInteractor ) override;
 	virtual void AddToVRRadialMenuActionsMenu( class IMeshEditorModeUIContract& MeshEditorMode, class FMenuBuilder& MenuBuilder, TSharedPtr<FUICommandList> CommandList, const FName TEMPHACK_StyleSetName, class UVREditorMode* VRMode ) override;
diff --git a/Engine/Plugins/Editor/MeshEditor/Source/PolygonModeling/ExtrudePolygon.h b/Engine/Plugins/Editor/MeshEditor/Source/PolygonModeling/ExtrudePolygon.h
--- a/Engine/Plugins/Editor/MeshEditor/Source/PolygonModeling/ExtrudePolygon.h
+++ b/Engine/Plugins/Editor/MeshEditor/Source/PolygonModeling/ExtrudePolygon.h
@@ -8,7 +8,7 @@
 
 /** Extrudes the polygon along an axis */
 UCLASS()
-class UExtrudePolygonCommand : public UMeshEditorPolygonCommand
+class UExtrudePolygonCommand : public UMeshEditorEditCommand
 {
 	GENERATED_BODY()
 
@@ -18,7 +18,6 @@ protected:
 		: ExtrudePolygonAxisOrigin( FVector::ZeroVector ),
 		  ExtrudePolygonAxisDirection( FVector::ZeroVector )
 	{
-		bIsMode = true;
 		UndoText = NSLOCTEXT( "MeshEditor", "UndoExtrudePolygon", "Extrude Polygon" );
 		bNeedsHoverLocation = false;
 		bNeedsDraggingInitiated = false;
@@ -35,6 +34,10 @@ protected:
 		float& OutExtrudeDistance );
 
 	// Overrides
+	virtual EEditableMeshElementType GetElementType() const override
+	{
+		return EEditableMeshElementType::Polygon;
+	}
 	virtual void RegisterUICommand( class FBindingContext* BindingContext ) override;
 	virtual bool TryStartingToDrag( IMeshEditorModeEditingContract& MeshEditorMode, class UViewportInteractor* ViewportInteractor ) override;
 	virtual void ApplyDuringDrag( IMeshEditorModeEditingContract& MeshEditorMode, class UViewportInteractor* ViewportInteractor ) override;
diff --git a/Engine/Plugins/Editor/MeshEditor/Source/PolygonModeling/HardenOrSoftenEdge.h b/Engine/Plugins/Editor/MeshEditor/Source/PolygonModeling/HardenOrSoftenEdge.h
--- a/Engine/Plugins/Editor/MeshEditor/Source/PolygonModeling/HardenOrSoftenEdge.h
+++ b/Engine/Plugins/Editor/MeshEditor/Source/PolygonModeling/HardenOrSoftenEdge.h
@@ -8,13 +8,17 @@
 
 /** Makes an edge hard */
 UCLASS()
-class UHardenEdgeCommand : public UMeshEditorEdgeCommand
+class UHardenEdgeCommand : public UMeshEditorInstantCommand
 {
 	GENERATED_BODY()
 
 protected:
 
 	// Overrides
+	virtual EEditableMeshElementType GetElementType() const override
+	{
+		return EEditableMeshElementType::Edge;
+	}
 	virtual void RegisterUICommand( class FBindingContext* BindingContext ) override;
 	virtual void Execute( class IMeshEditorModeEditingContract& MeshEditorMode ) override;
 
@@ -23,13 +27,17 @@ protected:
 
 /** Makes an edge soft */
 UCLASS()
-class USoftenEdgeCommand : public UMeshEditorEdgeCommand
+class USoftenEdgeCommand : public UMeshEditorInstantCommand
 {
 	GENERATED_BODY()
 
 protected:
 
 	// Overrides
+	virtual EEditableMeshElementType GetElementType() const override
+	{
+		return EEditableMeshElementType::Edge;
+	}
 	virtual void RegisterUICommand( class FBindingContext* BindingContext ) override;
 	virtual void Execute( class IMeshEditorModeEditingContract& MeshEditorMode ) override;
 
diff --git a/Engine/Plugins/Editor/MeshEditor/Source/PolygonModeling/InsertEdgeLoop.h b/Engine/Plugins/Editor/MeshEditor/Source/PolygonModeling/InsertEdgeLoop.h
--- a/Engine/Plugins/Editor/MeshEditor/Source/PolygonModeling/InsertEdgeLoop.h
+++ b/Engine/Plugins/Editor/MeshEditor/Source/PolygonModeling/InsertEdgeLoop.h
@@ -8,7 +8,7 @@
 
 /** With an edge selected, inserts a loop of edge perpendicular to that edge while dragging */
 UCLASS()
-class UInsertEdgeLoopCommand : public UMeshEditorEdgeCommand
+class UInsertEdgeLoopCommand : public UMeshEditorEditCommand
 {
 	GENERATED_BODY()
 
@@ -16,13 +16,16 @@ protected:
 
 	UInsertEdgeLoopCommand()
 	{
-		bIsMode = true;
 		UndoText = NSLOCTEXT( "MeshEditor", "UndoInsertEdgeLoop", "Insert Edge Loop" );
 		bNeedsHoverLocation = true;
 		bNeedsDraggingInitiated = false;
 	}
 
 	// Overrides
+	virtual EEditableMeshElementType GetElementType() const override
+	{
+		return EEditableMeshElementType::Edge;
+	}
 	virtual void RegisterUICommand( class FBindingContext* BindingContext ) override;
 	virtual void ApplyDuringDrag( IMeshEditorModeEditingContract& MeshEditorMode, class UViewportInteractor* ViewportInteractor ) override;
 	virtual void AddToVRRadialMenuActionsMenu( class IMeshEditorModeUIContract& MeshEditorMode, class FMenuBuilder& MenuBuilder, TSharedPtr<FUICommandList> CommandList, const FName TEMPHACK_StyleSetName, class UVREditorMode* VRMode ) override;
diff --git a/Engine/Plugins/Editor/MeshEditor/Source/PolygonModeling/QuadrangulateMesh.h b/Engine/Plugins/Editor/MeshEditor/Source/PolygonModeling/QuadrangulateMesh.h
--- a/Engine/Plugins/Editor/MeshEditor/Source/PolygonModeling/QuadrangulateMesh.h
+++ b/Engine/Plugins/Editor/MeshEditor/Source/PolygonModeling/QuadrangulateMesh.h
@@ -8,13 +8,17 @@
 
 /** Quadrangulates the currently selected mesh */
 UCLASS()
-class UQuadrangulateMeshCommand : public UMeshEditorCommonCommand
+class UQuadrangulateMeshCommand : public UMeshEditorInstantCommand
 {
 	GENERATED_BODY()
 
 protected:
 
 	// Overrides
+	virtual EEditableMeshElementType GetElementType() const override
+	{
+		return EEditableMeshElementType::Invalid;
+	}
 	virtual void RegisterUICommand( class FBindingContext* BindingContext ) override;
 	virtual void Execute( class IMeshEditorModeEditingContract& MeshEditorMode ) override;
 
diff --git a/Engine/Plugins/Editor/MeshEditor/Source/PolygonModeling/RemoveEdge.h b/Engine/Plugins/Editor/MeshEditor/Source/PolygonModeling/RemoveEdge.h
--- a/Engine/Plugins/Editor/MeshEditor/Source/PolygonModeling/RemoveEdge.h
+++ b/Engine/Plugins/Editor/MeshEditor/Source/PolygonModeling/RemoveEdge.h
@@ -8,13 +8,17 @@
 
 /** Attempts to remove the selected edge from the polygon, merging the adjacent polygons together */
 UCLASS()
-class URemoveEdgeCommand : public UMeshEditorEdgeCommand
+class URemoveEdgeCommand : public UMeshEditorInstantCommand
 {
 	GENERATED_BODY()
 
 protected:
 
 	// Overrides
+	virtual EEditableMeshElementType GetElementType() const override
+	{
+		return EEditableMeshElementType::Edge;
+	}
 	virtual void RegisterUICommand( class FBindingContext* BindingContext ) override;
 	virtual void Execute( class IMeshEditorModeEditingContract& MeshEditorMode ) override;
 	virtual void AddToVRRadialMenuActionsMenu( class IMeshEditorModeUIContract& MeshEditorMode, class FMenuBuilder& MenuBuilder, TSharedPtr<FUICommandList> CommandList, const FName TEMPHACK_StyleSetName, class UVREditorMode* VRMode ) override;
diff --git a/Engine/Plugins/Editor/MeshEditor/Source/PolygonModeling/RemoveVertex.h b/Engine/Plugins/Editor/MeshEditor/Source/PolygonModeling/RemoveVertex.h
--- a/Engine/Plugins/Editor/MeshEditor/Source/PolygonModeling/RemoveVertex.h
+++ b/Engine/Plugins/Editor/MeshEditor/Source/PolygonModeling/RemoveVertex.h
@@ -8,13 +8,17 @@
 
 /** Attempts to remove the selected vertex from the polygon it belongs to, keeping the polygon intact */
 UCLASS()
-class URemoveVertexCommand : public UMeshEditorVertexCommand
+class URemoveVertexCommand : public UMeshEditorInstantCommand
 {
 	GENERATED_BODY()
 
 protected:
 
 	// Overrides
+	virtual EEditableMeshElementType GetElementType() const override
+	{
+		return EEditableMeshElementType::Vertex;
+	}
 	virtual void RegisterUICommand( class FBindingContext* BindingContext ) override;
 	virtual void Execute( class IMeshEditorModeEditingContract& MeshEditorMode ) override;
 	virtual void AddToVRRadialMenuActionsMenu( class IMeshEditorModeUIContract& MeshEditorMode, class FMenuBuilder& MenuBuilder, TSharedPtr<FUICommandList> CommandList, const FName TEMPHACK_StyleSetName, class UVREditorMode* VRMode ) override;
diff --git a/Engine/Plugins/Editor/MeshEditor/Source/PolygonModeling/SplitEdge.h b/Engine/Plugins/Editor/MeshEditor/Source/PolygonModeling/SplitEdge.h
--- a/Engine/Plugins/Editor/MeshEditor/Source/PolygonModeling/SplitEdge.h
+++ b/Engine/Plugins/Editor/MeshEditor/Source/PolygonModeling/SplitEdge.h
@@ -8,7 +8,7 @@
 
 /** Inserts a vertex along an edge and allows the user to move it around */
 UCLASS()
-class USplitEdgeCommand : public UMeshEditorEdgeCommand
+class USplitEdgeCommand : public UMeshEditorEditCommand
 {
 	GENERATED_BODY()
 
@@ -18,13 +18,16 @@ protected:
 	  : SplitEdgeMeshesAndEdgesToSplit(),
 	    SplitEdgeSplitList()
 	{
-		bIsMode = true;
 		UndoText = NSLOCTEXT( "MeshEditor", "UndoSplitEdge", "Split Edge" );
 		bNeedsHoverLocation = false;
 		bNeedsDraggingInitiated = true;
 	}
 
 	// Overrides
+	virtual EEditableMeshElementType GetElementType() const override
+	{
+		return EEditableMeshElementType::Edge;
+	}
 	virtual void RegisterUICommand( class FBindingContext* BindingContext ) override;
 	virtual bool TryStartingToDrag( IMeshEditorModeEditingContract& MeshEditorMode, class UViewportInteractor* ViewportInteractor ) override;
 	virtual void ApplyDuringDrag( IMeshEditorModeEditingContract& MeshEditorMode, class UViewportInteractor* ViewportInteractor ) override;
diff --git a/Engine/Plugins/Editor/MeshEditor/Source/PolygonModeling/TessellatePolygon.h b/Engine/Plugins/Editor/MeshEditor/Source/PolygonModeling/TessellatePolygon.h
--- a/Engine/Plugins/Editor/MeshEditor/Source/PolygonModeling/TessellatePolygon.h
+++ b/Engine/Plugins/Editor/MeshEditor/Source/PolygonModeling/TessellatePolygon.h
@@ -8,13 +8,17 @@
 
 /** Tessellates selected polygons into smaller polygons */
 UCLASS()
-class UTessellatePolygonCommand : public UMeshEditorPolygonCommand
+class UTessellatePolygonCommand : public UMeshEditorInstantCommand
 {
 	GENERATED_BODY()
 
 protected:
 
 	// Overrides
+	virtual EEditableMeshElementType GetElementType() const override
+	{
+		return EEditableMeshElementType::Polygon;
+	}
 	virtual void RegisterUICommand( class FBindingContext* BindingContext ) override;
 	virtual void Execute( class IMeshEditorModeEditingContract& MeshEditorMode ) override;
 
