Mesh Editor: Fixed regression with Draw Vertices not working

diff --git a/Engine/Plugins/Editor/MeshEditor/Source/MeshEditor/MeshEditorMode.cpp b/Engine/Plugins/Editor/MeshEditor/Source/MeshEditor/MeshEditorMode.cpp
--- a/Engine/Plugins/Editor/MeshEditor/Source/MeshEditor/MeshEditorMode.cpp
+++ b/Engine/Plugins/Editor/MeshEditor/Source/MeshEditor/MeshEditorMode.cpp
@@ -3002,7 +3002,7 @@ void FMeshEditorMode::UpdateActiveAction( const bool bIsActionFinishing )
 			{
 				// Currently adds new vertices to whichever editable mesh is currently selected
 				Component = SelectedComponentsAndEditableMeshes[ 0 ].Component;
-				EditableMesh = SelectedComponentsAndEditableMeshes[ 1 ].EditableMesh;
+				EditableMesh = SelectedComponentsAndEditableMeshes[ 0 ].EditableMesh;
 				// @todo mesheditor: allow multiple selected meshes? What should this do?
 
 				check( EditableMesh );
