Mesh Editor: Fixed crash with FMeshElement::ToString() when component was invalid

diff --git a/Engine/Plugins/Editor/MeshEditor/Source/MeshEditor/Public/MeshElement.h b/Engine/Plugins/Editor/MeshEditor/Source/MeshEditor/Public/MeshElement.h
--- a/Engine/Plugins/Editor/MeshEditor/Source/MeshEditor/Public/MeshElement.h
+++ b/Engine/Plugins/Editor/MeshEditor/Source/MeshEditor/Public/MeshElement.h
@@ -172,7 +172,7 @@ struct FMeshElement
 	{
 		return FString::Printf(
 			TEXT( "Component:%s, %s" ),
-			*Component->GetName(),
+			Component.IsValid() ? *Component->GetName() : TEXT( "<Invalid>" ),
 			*ElementAddress.ToString() );
 	}
 };
