Mesh Editor: Interactor hover now updates before ticking within the frame
- Now that VWI is ticked before EdModes, we're able to update our hover state before we tick everything
- We now also make sure that preview changes are rolled back before we update hover state.  We never want to hover over previewed meshes!
- Also, added a new function to get edges for all selected polygons.  This is useful for features that operate on edges while only polygons are selected

diff --git a/Engine/Plugins/Editor/MeshEditor/Source/MeshEditor/MeshEditorMode.cpp b/Engine/Plugins/Editor/MeshEditor/Source/MeshEditor/MeshEditorMode.cpp
--- a/Engine/Plugins/Editor/MeshEditor/Source/MeshEditor/MeshEditorMode.cpp
+++ b/Engine/Plugins/Editor/MeshEditor/Source/MeshEditor/MeshEditorMode.cpp
@@ -997,15 +997,6 @@ void FMeshEditorMode::Tick( FEditorViewportClient* ViewportClient, float DeltaTi
 	}
 
 
-	// Clear hover for all interactors.  We'll re-detect what's hovered every frame when OnViewportInteractionHoverUpdate() is called.
-	for( FMeshEditorInteractorData& MeshEditorInteractorData : MeshEditorInteractorDatas )
-	{
-		MeshEditorInteractorData.PreviouslyHoveredMeshElement = MeshEditorInteractorData.HoveredMeshElement;
-		MeshEditorInteractorData.HoveredMeshElement = FMeshElement();
-		MeshEditorInteractorData.HoverLocation = FVector::ZeroVector;
-	}
-
-
 	// Hide the transform gizmo while we're doing things.  It actually will get in the way of our hit tests!
 	{
 		const EEditableMeshElementType SelectedMeshElementType = GetSelectedMeshElementType();
@@ -1215,6 +1206,52 @@ const UMeshEditorAssetContainer& FMeshEditorMode::GetAssetContainer() const
 	return *AssetContainer;
 }
 
+
+void FMeshEditorMode::GetSelectedMeshesAndPolygonsPerimeterEdges( TMap<UEditableMesh*, TArray<FMeshElement>>& OutMeshesAndPolygonsEdges )
+{
+	OutMeshesAndPolygonsEdges.Reset();
+
+	static TMap<UEditableMesh*, TArray<FMeshElement>> MeshesAndPolygons;
+	GetSelectedMeshesAndElements( EEditableMeshElementType::Polygon, /* Out */ MeshesAndPolygons );
+
+	for( auto& MeshAndPolygonElements : MeshesAndPolygons )
+	{
+		UEditableMesh* EditableMesh = MeshAndPolygonElements.Key;
+
+		static TArray<FEdgeID> UniqueSelectedEdgeIDs;
+		UniqueSelectedEdgeIDs.Reset();
+
+		FMeshElement FirstPolygonElement;
+		for( const FMeshElement& PolygonElement : MeshAndPolygonElements.Value )
+		{
+			if( !FirstPolygonElement.IsValidMeshElement() )
+			{
+				FirstPolygonElement = PolygonElement;
+			}
+
+			static TArray<FEdgeID> PerimeterEdgeIDs;
+			EditableMesh->GetPolygonPerimeterEdges( FPolygonRef( FSectionID( PolygonElement.ElementAddress.SectionID ), FPolygonID( PolygonElement.ElementAddress.ElementID ) ), /* Out */ PerimeterEdgeIDs );
+
+			for( const FEdgeID PerimeterEdgeID : PerimeterEdgeIDs )
+			{
+				UniqueSelectedEdgeIDs.AddUnique( PerimeterEdgeID );	// Unique add, because polygons can share edges
+			}
+		}
+
+		TArray<FMeshElement>& EdgeElementsToFill = OutMeshesAndPolygonsEdges.Add( EditableMesh, TArray<FMeshElement>() );
+		for( const FEdgeID EdgeID : UniqueSelectedEdgeIDs )
+		{
+			FMeshElement EdgeElement;
+			EdgeElement.Component = FirstPolygonElement.Component;
+			EdgeElement.ElementAddress.SubMeshAddress = FirstPolygonElement.ElementAddress.SubMeshAddress;
+			EdgeElement.ElementAddress.ElementID = EdgeID;
+			EdgeElement.ElementAddress.ElementType = EEditableMeshElementType::Edge;
+
+			EdgeElementsToFill.Add( EdgeElement );
+		}
+	}
+}
+
 void FMeshEditorMode::SelectMeshElements( const TArray<FMeshElement>& MeshElementsToSelect )
 {
 	if( MeshElementsToSelect.Num() > 0 )
@@ -2573,10 +2610,18 @@ FMeshEditorMode::FMeshEditorInteractorData& FMeshEditorMode::GetMeshEditorIntera
 
 void FMeshEditorMode::OnViewportInteractionHoverUpdate( UViewportInteractor* ViewportInteractor, FVector& OutHoverImpactPoint, bool& bWasHandled )
 {
+	FMeshEditorInteractorData& MeshEditorInteractorData = GetMeshEditorInteractorData( ViewportInteractor );
+	MeshEditorInteractorData.PreviouslyHoveredMeshElement = MeshEditorInteractorData.HoveredMeshElement;
+	MeshEditorInteractorData.HoveredMeshElement = FMeshElement();
+	MeshEditorInteractorData.HoverLocation = FVector::ZeroVector;
+
+	// Make sure there are no outstanding changes being previewed.  Usually, OnViewportInteractionHoverUpdate() will be the first function
+	// called on our class each frame.  We definitely don't want to do hover testing against the mesh we were previewing at the end of
+	// the last frame.  So let's roll those changes back first thing.
+	RollbackPreviewChanges();
+
 	if( !bWasHandled )
 	{
-		FMeshEditorInteractorData& MeshEditorInteractorData = GetMeshEditorInteractorData( ViewportInteractor );
-
 		MeshEditorInteractorData.bGrabberSphereIsValid = ViewportInteractor->GetGrabberSphere( /* Out */ MeshEditorInteractorData.GrabberSphere );
 		MeshEditorInteractorData.bLaserIsValid = ViewportInteractor->GetLaserPointer( /* Out */ MeshEditorInteractorData.LaserStart, /* Out */ MeshEditorInteractorData.LaserEnd );
 
diff --git a/Engine/Plugins/Editor/MeshEditor/Source/MeshEditor/MeshEditorMode.h b/Engine/Plugins/Editor/MeshEditor/Source/MeshEditor/MeshEditorMode.h
--- a/Engine/Plugins/Editor/MeshEditor/Source/MeshEditor/MeshEditorMode.h
+++ b/Engine/Plugins/Editor/MeshEditor/Source/MeshEditor/MeshEditorMode.h
@@ -191,6 +191,7 @@ protected:
 	{
 		GetSelectedMeshesAndElements( EEditableMeshElementType::Polygon, /* Out */ OutMeshesAndPolygons );
 	}
+	virtual void GetSelectedMeshesAndPolygonsPerimeterEdges( TMap<UEditableMesh*, TArray<FMeshElement>>& OutMeshesAndPolygonsEdges ) override;
 	virtual const TArray<UEditableMesh*>& GetSelectedEditableMeshes() const override { return SelectedEditableMeshes; }
 	virtual const TArray<UEditableMesh*>& GetSelectedEditableMeshes() override { return SelectedEditableMeshes; }
 	virtual void SelectMeshElements( const TArray<FMeshElement>& MeshElementsToSelect ) override;
diff --git a/Engine/Plugins/Editor/MeshEditor/Source/MeshEditor/Public/IMeshEditorModeEditingContract.h b/Engine/Plugins/Editor/MeshEditor/Source/MeshEditor/Public/IMeshEditorModeEditingContract.h
--- a/Engine/Plugins/Editor/MeshEditor/Source/MeshEditor/Public/IMeshEditorModeEditingContract.h
+++ b/Engine/Plugins/Editor/MeshEditor/Source/MeshEditor/Public/IMeshEditorModeEditingContract.h
@@ -32,6 +32,7 @@ public:
 	virtual void GetSelectedMeshesAndVertices( TMap<class UEditableMesh*, TArray<FMeshElement>>& OutMeshesAndVertices ) = 0;
 	virtual void GetSelectedMeshesAndEdges( TMap<class UEditableMesh*, TArray<FMeshElement>>& OutMeshesAndEdges ) = 0;
 	virtual void GetSelectedMeshesAndPolygons( TMap<class UEditableMesh*, TArray<FMeshElement>>& OutMeshesAndPolygons ) = 0;
+	virtual void GetSelectedMeshesAndPolygonsPerimeterEdges( TMap<class UEditableMesh*, TArray<FMeshElement>>& OutMeshesAndPolygonsEdges ) = 0;
 
 	/** Gets all of the editable meshes that are currently selected (const) */
 	virtual const TArray<class UEditableMesh*>& GetSelectedEditableMeshes() const = 0;
