Mesh Editor: Split Polygon
- Added first pass at new "Split Polygon" feature.
- You can select any mesh element, then click and drag to split an adjacent polygon by dragging to an opposing edge or vertex.
- Currently it uses the closest edge to the point the mouse is hovering over.  This will be improved in the future.
- Undo/Redo is fully supported.  Works in VR Mode, also.

diff --git a/Engine/Plugins/Editor/MeshEditor/Source/PolygonModeling/SplitPolygon.cpp b/Engine/Plugins/Editor/MeshEditor/Source/PolygonModeling/SplitPolygon.cpp
new file mode 100644
--- /dev/null
+++ b/Engine/Plugins/Editor/MeshEditor/Source/PolygonModeling/SplitPolygon.cpp
@@ -0,0 +1,344 @@
+// Copyright 1998-2016 Epic Games, Inc. All Rights Reserved.
+
+#include "SplitPolygon.h"
+#include "IMeshEditorModeEditingContract.h"
+#include "IMeshEditorModeUIContract.h"
+#include "UICommandInfo.h"
+#include "EditableMesh.h"
+#include "MeshElement.h"
+#include "MultiBoxBuilder.h"
+#include "UICommandList.h"
+#include "ViewportInteractor.h"
+
+#define LOCTEXT_NAMESPACE "MeshEditorMode"
+
+
+void USplitPolygonFromVertexCommand::RegisterUICommand( FBindingContext* BindingContext )
+{
+	UI_COMMAND_EXT( BindingContext, /* Out */ UICommandInfo, "SplitPolygonFromVertex", "Split Polygon Mode", "Set the primary action to split a polygon.", EUserInterfaceActionType::RadioButton, FInputChord() );
+}
+
+
+void USplitPolygonFromEdgeCommand::RegisterUICommand( FBindingContext* BindingContext )
+{
+	UI_COMMAND_EXT( BindingContext, /* Out */ UICommandInfo, "SplitPolygonFromEdge", "Split Polygon Mode", "Set the primary action to split a polygon.", EUserInterfaceActionType::RadioButton, FInputChord() );
+}
+
+
+void USplitPolygonFromPolygonCommand::RegisterUICommand( FBindingContext* BindingContext )
+{
+	UI_COMMAND_EXT( BindingContext, /* Out */ UICommandInfo, "SplitPolygonFromPolygon", "Split Polygon Mode", "Set the primary action to split a polygon.", EUserInterfaceActionType::RadioButton, FInputChord() );
+}
+
+
+bool USplitPolygonCommand::TryStartingToDrag( IMeshEditorModeEditingContract& MeshEditorMode, UViewportInteractor* ViewportInteractor )
+{
+	Component = nullptr;
+	EditableMesh = nullptr;
+	StartingEdgeID = FEdgeID::Invalid;
+	StartingVertexID = FVertexID::Invalid;
+	EdgeSplit = 0.0f;
+
+	// Figure out what to split
+	static TMap<class UEditableMesh*, TArray<FMeshElement>> SelectedMeshesAndEdges;
+	MeshEditorMode.GetSelectedMeshesAndEdges( /* Out */ SelectedMeshesAndEdges );
+	if( SelectedMeshesAndEdges.Num() == 0 )
+	{
+		// No edges were selected directly, so try selected polygons' edges
+		MeshEditorMode.GetSelectedMeshesAndPolygonsPerimeterEdges( /* Out */ SelectedMeshesAndEdges );
+	}
+
+	if( SelectedMeshesAndEdges.Num() > 0 )
+	{
+		for( auto& MeshAndEdges : SelectedMeshesAndEdges )
+		{
+			UEditableMesh* EdgeEditableMesh = MeshAndEdges.Key;
+			const TArray<FMeshElement>& EdgeElements = MeshAndEdges.Value;
+
+			// Figure out where to split
+			FEdgeID ClosestEdgeID = FEdgeID::Invalid;
+			float Split = 0.0f;
+			const bool bFoundSplit = MeshEditorMode.FindEdgeSplitUnderInteractor( ViewportInteractor, EdgeEditableMesh, EdgeElements, /* Out */ ClosestEdgeID, /* Out */ Split );
+
+			// @todo now: We need to only support a single selected edge or vertex, or find the one they actually clicked on!!
+			if( bFoundSplit )
+			{
+				// OK, we have an edge position to start dragging from!
+				Component = EdgeElements[ 0 ].Component.Get();
+				EditableMesh = EdgeEditableMesh;
+				StartingEdgeID = ClosestEdgeID;
+				EdgeSplit = Split;
+				break;
+			}
+		}
+	}
+	else
+	{
+		static TMap<class UEditableMesh*, TArray<FMeshElement>> SelectedMeshesAndVertices;
+		MeshEditorMode.GetSelectedMeshesAndVertices( /* Out */ SelectedMeshesAndVertices );
+
+		for( auto& MeshAndVertices : SelectedMeshesAndVertices )
+		{
+			UEditableMesh* VertexEditableMesh = MeshAndVertices.Key;
+			const TArray<FMeshElement>& VertexElements = MeshAndVertices.Value;
+
+			for( const FMeshElement& VertexElement : VertexElements )
+			{
+				// OK, we have a vertex to start dragging from!
+				Component = VertexElements[ 0 ].Component.Get();
+				EditableMesh = VertexEditableMesh;
+				StartingVertexID = FVertexID( VertexElement.ElementAddress.ElementID );
+				break;
+			}
+		}
+	}
+
+	return EditableMesh != nullptr;
+}
+
+
+void USplitPolygonCommand::ApplyDuringDrag( IMeshEditorModeEditingContract& MeshEditorMode, UViewportInteractor* ViewportInteractor )
+{
+	if( EditableMesh != nullptr )
+	{
+		check( Component != nullptr );
+
+ 		MeshEditorMode.DeselectAllMeshElements();
+ 
+ 		static TArray<FMeshElement> MeshElementsToSelect;
+ 		MeshElementsToSelect.Reset();
+ 
+		verify( !EditableMesh->AnyChangesToUndo() );
+
+		// We'll always be trying to split one of the polygons that share the starting vertex.  But we need to figure
+		// out which one of those polygons the user is hovering over (either the polygon itself, or one of it's edges
+		// or vertices.)
+		static TArray<FPolygonRef> CandidatePolygonRefs;
+		CandidatePolygonRefs.Reset();
+		if( StartingEdgeID != FEdgeID::Invalid )
+		{
+			EditableMesh->GetEdgeConnectedPolygons( StartingEdgeID, /* Out */ CandidatePolygonRefs );
+		}
+		else if( ensure( StartingVertexID != FVertexID::Invalid ) )
+		{
+			EditableMesh->GetVertexConnectedPolygons( StartingVertexID, /* Out */ CandidatePolygonRefs );
+		}
+
+		FPolygonRef PolygonToSplit = FPolygonRef::Invalid;
+		FVertexID ToVertexID = FVertexID::Invalid;
+
+		const FMeshElement& HoveredElement = MeshEditorMode.GetHoveredMeshElement( ViewportInteractor );
+		if( HoveredElement.IsValidMeshElement() )	// @todo now: When using this feature, we need to allow hovering over ANY element type
+		{
+			// Make sure we're hovered over the same mesh
+			if( HoveredElement.ElementAddress.SubMeshAddress == EditableMesh->GetSubMeshAddress() )
+			{
+				if( HoveredElement.ElementAddress.ElementType == EEditableMeshElementType::Vertex )
+				{
+					const FVertexID CandidateVertexID( HoveredElement.ElementAddress.ElementID );
+
+					// Don't allow dragging to the vertex we started from
+					if( StartingVertexID != CandidateVertexID )
+					{
+						// Make sure the hovered vertex isn't already connected to our starting vertex by an edge
+						bool bIsDisqualified = false;
+						if( StartingEdgeID != FEdgeID::Invalid )
+						{
+							static TArray<FEdgeID> AdjacentEdgeIDs;
+							EditableMesh->GetVertexConnectedEdges( CandidateVertexID, /* Out */ AdjacentEdgeIDs );
+
+							bIsDisqualified = AdjacentEdgeIDs.Contains( StartingEdgeID );
+						}
+						else if( ensure( StartingVertexID != FVertexID::Invalid ) )
+						{
+							static TArray<FVertexID> AdjacentVertexIDs;
+							EditableMesh->GetVertexAdjacentVertices( CandidateVertexID, /* Out */ AdjacentVertexIDs );
+
+							bIsDisqualified = AdjacentVertexIDs.Contains( StartingVertexID );
+						}
+
+						if( !bIsDisqualified )	// @todo now: Also check whether vertex is colinear with connected edges and skip if so?  Or should that be part of SplitPolygon?
+						{
+							// Make sure the hovered vertex is part of our candidate polygon set
+							for( const FPolygonRef CandidatePolygonRef : CandidatePolygonRefs )
+							{
+								static TArray<FVertexID> CandidatePolygonPerimeterVertexIDs;
+								EditableMesh->GetPolygonPerimeterVertices( CandidatePolygonRef, /* Out */ CandidatePolygonPerimeterVertexIDs );
+
+								if( CandidatePolygonPerimeterVertexIDs.Contains( CandidateVertexID ) )
+								{
+									// OK, we found a polygon we can split and a target vertex
+									PolygonToSplit = CandidatePolygonRef;
+									ToVertexID = CandidateVertexID;
+									
+									break;
+								}
+							}
+						}
+					}
+				}
+				else if( HoveredElement.ElementAddress.ElementType == EEditableMeshElementType::Edge ||
+					     HoveredElement.ElementAddress.ElementType == EEditableMeshElementType::Polygon )
+				{
+					// Figure out where over the edge we're hovering
+
+					// Gather all of the candidate edges we can join to.  These are just the edges of the candidate polygon, excluding
+					// any edges that we started on
+					static TArray<FEdgeID> CandidateEdgeIDs;
+					{
+						CandidateEdgeIDs.Reset();
+						for( const FPolygonRef CandidatePolygonRef : CandidatePolygonRefs )
+						{
+							static TArray<FEdgeID> PerimeterEdgeIDs;
+							EditableMesh->GetPolygonPerimeterEdges( CandidatePolygonRef, /* Out */ PerimeterEdgeIDs );
+
+							for( const FEdgeID PerimeterEdgeID : PerimeterEdgeIDs )
+							{
+								bool bIsDisqualified = false;
+								{
+									if( StartingEdgeID != FEdgeID::Invalid )
+									{
+										if( StartingEdgeID == PerimeterEdgeID )
+										{
+											bIsDisqualified = true;
+										}
+									}
+									else if( ensure( StartingVertexID != FVertexID::Invalid ) )
+									{
+										static TArray<FEdgeID> AdjacentEdgeIDs;
+										EditableMesh->GetVertexConnectedEdges( StartingVertexID, /* Out */ AdjacentEdgeIDs );
+										if( AdjacentEdgeIDs.Contains( PerimeterEdgeID ) )
+										{
+											bIsDisqualified = true;
+										}
+									}
+								}
+
+								// Never split an edge that's directly connected to the starting vertex
+								if( !bIsDisqualified )
+								{
+									CandidateEdgeIDs.AddUnique( PerimeterEdgeID );
+								}
+							}
+						}
+					}
+
+					if( CandidateEdgeIDs.Num() > 0 )
+					{
+						static TArray<FMeshElement> CandidateEdgeElements;
+						CandidateEdgeElements.Reset();
+						for( const FEdgeID CandidateEdgeID : CandidateEdgeIDs )
+						{
+							FMeshElement CandidateEdgeElement;
+							CandidateEdgeElement.Component = Component;
+							CandidateEdgeElement.ElementAddress.SubMeshAddress = EditableMesh->GetSubMeshAddress();
+							CandidateEdgeElement.ElementAddress.ElementType = EEditableMeshElementType::Edge;
+							CandidateEdgeElement.ElementAddress.ElementID = CandidateEdgeID;
+
+							CandidateEdgeElements.Add( CandidateEdgeElement );
+						}
+
+						// Figure out where to split
+						FEdgeID TargetEdgeID = FEdgeID::Invalid;
+						float OtherEdgeSplit = 0.0f;
+						const bool bFoundSplit = MeshEditorMode.FindEdgeSplitUnderInteractor( ViewportInteractor, EditableMesh, CandidateEdgeElements, /* Out */ TargetEdgeID, /* Out */ OtherEdgeSplit );
+						if( bFoundSplit )
+						{
+							// Figure out which candidate polygon to split
+							for( const FPolygonRef CandidatePolygonRef : CandidatePolygonRefs )
+							{
+								static TArray<FEdgeID> PerimeterEdgeIDs;
+								EditableMesh->GetPolygonPerimeterEdges( CandidatePolygonRef, /* Out */ PerimeterEdgeIDs );
+								if( PerimeterEdgeIDs.Contains( TargetEdgeID ) )
+								{
+									PolygonToSplit = CandidatePolygonRef;
+									break;
+								}
+							}
+							check( PolygonToSplit != FPolygonRef::Invalid );
+
+							// Go ahead and split the target edge
+							static TArray<FVertexID> NewVertexIDs;
+							NewVertexIDs.Reset();
+
+							static TArray<float> SplitEdgeSplitList;
+							SplitEdgeSplitList.SetNumUninitialized( 1 );
+							SplitEdgeSplitList[ 0 ] = OtherEdgeSplit;
+
+							EditableMesh->SplitEdge( TargetEdgeID, SplitEdgeSplitList, /* Out */ NewVertexIDs );
+							check( NewVertexIDs.Num() == 1 );
+
+							// Great!  We now have a vertex to connect with
+							ToVertexID = NewVertexIDs[ 0 ];
+						}
+					}
+				}
+			}
+		}
+
+		if( PolygonToSplit != FPolygonRef::Invalid )
+		{
+			check( ToVertexID != FVertexID::Invalid );
+
+			FVertexID FromVertexID = FVertexID::Invalid;
+			if( StartingEdgeID != FEdgeID::Invalid )
+			{
+				// First, split the edge
+				static TArray<FVertexID> NewVertexIDs;
+				NewVertexIDs.Reset();
+
+				static TArray<float> SplitEdgeSplitList;
+				SplitEdgeSplitList.SetNumUninitialized( 1 );
+				SplitEdgeSplitList[ 0 ] = EdgeSplit;
+
+				EditableMesh->SplitEdge( StartingEdgeID, SplitEdgeSplitList, /* Out */ NewVertexIDs );
+				check( NewVertexIDs.Num() == 1 );
+
+				FromVertexID = NewVertexIDs[ 0 ];
+			}
+			else if( ensure( StartingVertexID != FVertexID::Invalid ) )
+			{
+				FromVertexID = StartingVertexID;
+			}
+			check( FromVertexID != FVertexID::Invalid );
+
+			static TArray< FPolygonToSplit > PolygonsToSplit;
+			PolygonsToSplit.Reset();
+			PolygonsToSplit.SetNum( 1, false );
+
+			PolygonsToSplit[ 0 ].PolygonRef = PolygonToSplit;
+			FVertexPair& VertexPair = *new( PolygonsToSplit[ 0 ].VertexPairsToSplitAt ) FVertexPair();
+			VertexPair.VertexID0 = FromVertexID;
+			VertexPair.VertexID1 = ToVertexID;
+
+			static TArray<FEdgeID> NewEdgeIDs;
+			EditableMesh->SplitPolygons( PolygonsToSplit, /* Out */ NewEdgeIDs );
+
+			// Select the new edges that were created
+			{
+				for( const FEdgeID NewEdgeID : NewEdgeIDs )
+				{
+					FMeshElement MeshElementToSelect;
+					{
+						MeshElementToSelect.Component = Component;
+						MeshElementToSelect.ElementAddress.SubMeshAddress = EditableMesh->GetSubMeshAddress();
+						MeshElementToSelect.ElementAddress.ElementType = EEditableMeshElementType::Edge;
+						MeshElementToSelect.ElementAddress.ElementID = NewEdgeID;
+					}
+		
+					MeshElementsToSelect.Add( MeshElementToSelect );
+				}
+			}
+		}
+
+ 
+		MeshEditorMode.TrackUndo( EditableMesh, EditableMesh->MakeUndo() );
+ 
+ 		MeshEditorMode.SelectMeshElements( MeshElementsToSelect );
+ 	}
+}
+
+
+
+#undef LOCTEXT_NAMESPACE
+
diff --git a/Engine/Plugins/Editor/MeshEditor/Source/PolygonModeling/SplitPolygon.h b/Engine/Plugins/Editor/MeshEditor/Source/PolygonModeling/SplitPolygon.h
new file mode 100644
--- /dev/null
+++ b/Engine/Plugins/Editor/MeshEditor/Source/PolygonModeling/SplitPolygon.h
@@ -0,0 +1,106 @@
+// Copyright 1998-2016 Epic Games, Inc. All Rights Reserved.
+
+#pragma once
+
+#include "MeshEditorCommands.h"
+#include "SplitPolygon.generated.h"
+
+
+/** Base class for polygon splitting */
+UCLASS( abstract )
+class USplitPolygonCommand : public UMeshEditorEditCommand
+{
+	GENERATED_BODY()
+
+protected:
+
+	USplitPolygonCommand()
+	  : Component( nullptr ),
+		EditableMesh( nullptr ),
+		StartingEdgeID( FEdgeID::Invalid ),
+		StartingVertexID( FVertexID::Invalid ),
+		EdgeSplit( 0.0f )
+	{
+		UndoText = NSLOCTEXT( "MeshEditor", "UndoSplitPolygon", "Split Polygon" );
+		bNeedsHoverLocation = true;
+		bNeedsDraggingInitiated = false;
+	}
+
+	// Overrides
+	virtual bool TryStartingToDrag( IMeshEditorModeEditingContract& MeshEditorMode, class UViewportInteractor* ViewportInteractor ) override;
+	virtual void ApplyDuringDrag( IMeshEditorModeEditingContract& MeshEditorMode, class UViewportInteractor* ViewportInteractor ) override;
+
+protected:
+
+	/** The component we're editing */
+	UPROPERTY()
+	class UPrimitiveComponent* Component;	// @todo now: Should be weak? 
+
+	/** The mesh we're editing */
+	UPROPERTY()
+	class UEditableMesh* EditableMesh;		// @todo now: Should be weak?
+
+	/** The edge we started our operation on */
+	FEdgeID StartingEdgeID;
+
+	/** The edge or vertex we started our split operation on */
+	FVertexID StartingVertexID;
+
+	/** If we're starting on an edge, the progress along that edge to start at */
+	float EdgeSplit;
+
+};
+
+
+/** Splits a polygon into two, starting with a vertex on that polygon */
+UCLASS()
+class USplitPolygonFromVertexCommand : public USplitPolygonCommand
+{
+	GENERATED_BODY()
+
+protected:
+
+	// Overrides
+	virtual EEditableMeshElementType GetElementType() const override
+	{
+		return EEditableMeshElementType::Vertex;
+	}
+	virtual void RegisterUICommand( class FBindingContext* BindingContext ) override;
+
+};
+
+
+/** Splits a polygon into two, starting with a point along an edge */
+UCLASS()
+class USplitPolygonFromEdgeCommand : public USplitPolygonCommand
+{
+	GENERATED_BODY()
+
+protected:
+
+	// Overrides
+	virtual EEditableMeshElementType GetElementType() const override
+	{
+		return EEditableMeshElementType::Edge;
+	}
+	virtual void RegisterUICommand( class FBindingContext* BindingContext ) override;
+
+};
+
+
+/** Splits a polygon into two, starting with an edge on a polygon */
+UCLASS()
+class USplitPolygonFromPolygonCommand : public USplitPolygonCommand
+{
+	GENERATED_BODY()
+
+protected:
+
+	// Overrides
+	virtual EEditableMeshElementType GetElementType() const override
+	{
+		return EEditableMeshElementType::Polygon;
+	}
+	virtual void RegisterUICommand( class FBindingContext* BindingContext ) override;
+
+};
