Mesh Editor: Improved Split Polygon tool
- Split Polygon will now figure out the best split line based off the direction you drag from the starting edge or vertex
- Fixed issue where wrong vertex to connect to could be selected
- You can now connect to the end of an edge that is next to a disqualified vertex

diff --git a/Engine/Plugins/Editor/MeshEditor/Source/PolygonModeling/SplitPolygon.cpp b/Engine/Plugins/Editor/MeshEditor/Source/PolygonModeling/SplitPolygon.cpp
--- a/Engine/Plugins/Editor/MeshEditor/Source/PolygonModeling/SplitPolygon.cpp
+++ b/Engine/Plugins/Editor/MeshEditor/Source/PolygonModeling/SplitPolygon.cpp
@@ -60,7 +60,6 @@ bool USplitPolygonCommand::TryStartingToDrag( IMeshEditorModeEditingContract& Me
 			float Split = 0.0f;
 			const bool bFoundSplit = MeshEditorMode.FindEdgeSplitUnderInteractor( ViewportInteractor, EdgeEditableMesh, EdgeElements, /* Out */ ClosestEdgeID, /* Out */ Split );
 
-			// @todo now: We need to only support a single selected edge or vertex, or find the one they actually clicked on!!
 			if( bFoundSplit )
 			{
 				// OK, we have an edge position to start dragging from!
@@ -68,6 +67,8 @@ bool USplitPolygonCommand::TryStartingToDrag( IMeshEditorModeEditingContract& Me
 				EditableMesh = EdgeEditableMesh;
 				StartingEdgeID = ClosestEdgeID;
 				EdgeSplit = Split;
+
+				// No need to search other meshes
 				break;
 			}
 		}
@@ -127,155 +128,195 @@ void USplitPolygonCommand::ApplyDuringDrag( IMeshEditorModeEditingContract& Mesh
 		FPolygonRef PolygonToSplit = FPolygonRef::Invalid;
 		FVertexID ToVertexID = FVertexID::Invalid;
 
-		const FMeshElement& HoveredElement = MeshEditorMode.GetHoveredMeshElement( ViewportInteractor );
-		if( HoveredElement.IsValidMeshElement() )	// @todo now: When using this feature, we need to allow hovering over ANY element type
+
 		{
-			// Make sure we're hovered over the same mesh
-			if( HoveredElement.ElementAddress.SubMeshAddress == EditableMesh->GetSubMeshAddress() )
+			FVector LaserPointerStart, LaserPointerEnd;
+			if( ViewportInteractor->GetLaserPointer( /* Out */ LaserPointerStart, /* Out */ LaserPointerEnd ) )
 			{
-				if( HoveredElement.ElementAddress.ElementType == EEditableMeshElementType::Vertex )
+				for( const FPolygonRef CandidatePolygonRef : CandidatePolygonRefs )
 				{
-					const FVertexID CandidateVertexID( HoveredElement.ElementAddress.ElementID );
+					const FTransform ComponentToWorld = Component->GetComponentToWorld();
 
-					// Don't allow dragging to the vertex we started from
-					if( StartingVertexID != CandidateVertexID )
+					FVector SplitStartLocation;
+					if( StartingVertexID != FVertexID::Invalid )
 					{
-						// Make sure the hovered vertex isn't already connected to our starting vertex by an edge
-						bool bIsDisqualified = false;
-						if( StartingEdgeID != FEdgeID::Invalid )
-						{
-							static TArray<FEdgeID> AdjacentEdgeIDs;
-							EditableMesh->GetVertexConnectedEdges( CandidateVertexID, /* Out */ AdjacentEdgeIDs );
+						SplitStartLocation = ComponentToWorld.TransformPosition( EditableMesh->GetVertexAttribute( StartingVertexID, UEditableMeshAttribute::VertexPosition(), 0 ) );
+					}
+					else if( ensure( StartingEdgeID != FEdgeID::Invalid ) )
+					{
+						FVertexID EdgeVertex0, EdgeVertex1;
+						EditableMesh->GetEdgeVertices( StartingEdgeID, /* Out */ EdgeVertex0, /* Out */ EdgeVertex1 );
 
-							bIsDisqualified = AdjacentEdgeIDs.Contains( StartingEdgeID );
-						}
-						else if( ensure( StartingVertexID != FVertexID::Invalid ) )
-						{
-							static TArray<FVertexID> AdjacentVertexIDs;
-							EditableMesh->GetVertexAdjacentVertices( CandidateVertexID, /* Out */ AdjacentVertexIDs );
+						const FVector EdgeVertex0Location = ComponentToWorld.TransformPosition( EditableMesh->GetVertexAttribute( EdgeVertex0, UEditableMeshAttribute::VertexPosition(), 0 ) );
+						const FVector EdgeVertex1Location = ComponentToWorld.TransformPosition( EditableMesh->GetVertexAttribute( EdgeVertex1, UEditableMeshAttribute::VertexPosition(), 0 ) ) ;
 
-							bIsDisqualified = AdjacentVertexIDs.Contains( StartingVertexID );
-						}
+						SplitStartLocation = FMath::Lerp( EdgeVertex0Location, EdgeVertex1Location, EdgeSplit );
+					}
 
-						if( !bIsDisqualified )	// @todo now: Also check whether vertex is colinear with connected edges and skip if so?  Or should that be part of SplitPolygon?
+
+					{
+						const FPlane PolygonPlane = EditableMesh->ComputePolygonPlane( CandidatePolygonRef ).TransformBy( ComponentToWorld.ToMatrixWithScale() );
+	
+						const FVector LaserImpactOnPolygonPlane = FMath::LinePlaneIntersection( LaserPointerStart, LaserPointerEnd, PolygonPlane );
+
+						// @todo mesheditor splitpolygon: Ideally this would be more "fuzzy", and allow the interactor to extend beyond the range of the polygon.  But it will make figuring out which polygon to split more tricky.
+						// @todo mesheditor urgent: Can crash with "Colinear points in FMath::ComputeBaryCentric2D()"  Needs repro.
+						static TArray<int32> PerimeterVertexIndices;
+						FVector TriangleVertexWeights;
+						if( EditableMesh->ComputeBarycentricWeightForPointOnPolygon( CandidatePolygonRef, ComponentToWorld.InverseTransformPosition( LaserImpactOnPolygonPlane ), /* Out */ PerimeterVertexIndices, /* Out */ TriangleVertexWeights ) )
 						{
-							// Make sure the hovered vertex is part of our candidate polygon set
-							for( const FPolygonRef CandidatePolygonRef : CandidatePolygonRefs )
+							const FVector SplitDirection = ( LaserImpactOnPolygonPlane - SplitStartLocation ).GetSafeNormal();
+
+							// Trace out within the polygon to figure out where the split should connect to
+							float ClosestEdgeDistance = TNumericLimits<float>::Max();
+							FEdgeID ClosestEdgeID = FEdgeID::Invalid;
+
+							static TArray<FEdgeID> PolygonPerimeterEdgeIDs;
+							EditableMesh->GetPolygonPerimeterEdges( CandidatePolygonRef, /* Out */ PolygonPerimeterEdgeIDs );
+
+							for( const FEdgeID TargetEdgeID : PolygonPerimeterEdgeIDs )
 							{
-								static TArray<FVertexID> CandidatePolygonPerimeterVertexIDs;
-								EditableMesh->GetPolygonPerimeterVertices( CandidatePolygonRef, /* Out */ CandidatePolygonPerimeterVertexIDs );
+								bool bIsDisqualified = false;
 
-								if( CandidatePolygonPerimeterVertexIDs.Contains( CandidateVertexID ) )
+								if( StartingEdgeID != FEdgeID::Invalid )
 								{
-									// OK, we found a polygon we can split and a target vertex
-									PolygonToSplit = CandidatePolygonRef;
-									ToVertexID = CandidateVertexID;
-									
-									break;
+									if( StartingEdgeID == TargetEdgeID )
+									{
+										// The edge we dragged from is disqualified as a target
+										bIsDisqualified = true;
+									}
+								}
+								else if( ensure( StartingVertexID != FVertexID::Invalid ) )
+								{
+									static TArray<FEdgeID> AdjacentEdgeIDs;
+									EditableMesh->GetVertexConnectedEdges( StartingVertexID, /* Out */ AdjacentEdgeIDs );
+									if( AdjacentEdgeIDs.Contains( TargetEdgeID ) )
+									{
+										// Never split an edge that's directly connected to the starting vertex
+										bIsDisqualified = true;
+									}
 								}
-							}
-						}
-					}
-				}
-				else if( HoveredElement.ElementAddress.ElementType == EEditableMeshElementType::Edge ||
-					     HoveredElement.ElementAddress.ElementType == EEditableMeshElementType::Polygon )
-				{
-					// Figure out where over the edge we're hovering
 
-					// Gather all of the candidate edges we can join to.  These are just the edges of the candidate polygon, excluding
-					// any edges that we started on
-					static TArray<FEdgeID> CandidateEdgeIDs;
-					{
-						CandidateEdgeIDs.Reset();
-						for( const FPolygonRef CandidatePolygonRef : CandidatePolygonRefs )
-						{
-							static TArray<FEdgeID> PerimeterEdgeIDs;
-							EditableMesh->GetPolygonPerimeterEdges( CandidatePolygonRef, /* Out */ PerimeterEdgeIDs );
+								if( !bIsDisqualified )
+								{
+									FVertexID EdgeVertex0, EdgeVertex1;
+									EditableMesh->GetEdgeVertices( TargetEdgeID, /* Out */ EdgeVertex0, /* Out */ EdgeVertex1 );
+
+									const FVector EdgeVertex0Location = ComponentToWorld.TransformPosition( EditableMesh->GetVertexAttribute( EdgeVertex0, UEditableMeshAttribute::VertexPosition(), 0 ) );
+									const FVector EdgeVertex1Location = ComponentToWorld.TransformPosition( EditableMesh->GetVertexAttribute( EdgeVertex1, UEditableMeshAttribute::VertexPosition(), 0 ) );
+
+									FVector ClosestPointOnSplitLine, ClosestPointOnEdge;
+									FMath::SegmentDistToSegmentSafe(
+										SplitStartLocation, SplitStartLocation + SplitDirection * 99999.0f,
+										EdgeVertex0Location, EdgeVertex1Location,
+										/* Out */ ClosestPointOnSplitLine,
+										/* Out */ ClosestPointOnEdge );
+
+									// Closest points should be the same if there was actually an intersection
+									if( ClosestPointOnSplitLine.Equals( ClosestPointOnEdge ) )
+									{
+										const float DistanceToEdgeImpact = ( ClosestPointOnEdge - SplitStartLocation ).Size();
+										if( DistanceToEdgeImpact < ClosestEdgeDistance )
+										{
+											ClosestEdgeDistance = DistanceToEdgeImpact;
+											ClosestEdgeID = TargetEdgeID;
+										}
+									}
+								}
+							}
 
-							for( const FEdgeID PerimeterEdgeID : PerimeterEdgeIDs )
+							if( ClosestEdgeID != FEdgeID::Invalid )
 							{
-								bool bIsDisqualified = false;
+								FVertexID EdgeVertex0, EdgeVertex1;
+								EditableMesh->GetEdgeVertices( ClosestEdgeID, /* Out */ EdgeVertex0, /* Out */ EdgeVertex1 );
+
+								const FVector EdgeVertex0Location = ComponentToWorld.TransformPosition( EditableMesh->GetVertexAttribute( EdgeVertex0, UEditableMeshAttribute::VertexPosition(), 0 ) );
+								const FVector EdgeVertex1Location = ComponentToWorld.TransformPosition( EditableMesh->GetVertexAttribute( EdgeVertex1, UEditableMeshAttribute::VertexPosition(), 0 ) );
+
+								const FVector ImpactOnEdge = SplitStartLocation + SplitDirection * ClosestEdgeDistance;
+
+								const float EdgeLength = ( EdgeVertex1Location - EdgeVertex0Location ).Size();
+								const float ImpactProgressAlongEdge = ( ImpactOnEdge - EdgeVertex0Location ).Size() / EdgeLength;
+
+								// If we're really close to one side or the other of an edge, try to connect with an existing vertex instead
+								const float EdgeProgressVertexSnapThreshold = 0.075f;	// @todo mesheditor splitpolygon: Should be actual 'fuzzy' distance consistent with MeshEditorMode, in world/screen units, not percentage of edge progress
+								const bool bIsCloseToEdgeVertex0 = ImpactProgressAlongEdge < EdgeProgressVertexSnapThreshold;
+								const bool bIsCloseToEdgeVertex1 = ImpactProgressAlongEdge > 1.0f - EdgeProgressVertexSnapThreshold;
+								if( bIsCloseToEdgeVertex0 || bIsCloseToEdgeVertex1 )
 								{
+									// Prefer connecting to a vertex
+									const FVertexID TargetVertexID = bIsCloseToEdgeVertex0 ? EdgeVertex0 : EdgeVertex1;
+
+									bool bIsDisqualified = false;
+
 									if( StartingEdgeID != FEdgeID::Invalid )
 									{
-										if( StartingEdgeID == PerimeterEdgeID )
+										FVertexID StartingEdgeVertex0, StartingEdgeVertex1;
+										EditableMesh->GetEdgeVertices( StartingEdgeID, /* Out */ StartingEdgeVertex0, /* Out */ StartingEdgeVertex1 );
+
+										if( TargetVertexID == StartingEdgeVertex0 || TargetVertexID == StartingEdgeVertex1 )
 										{
+											// We're dragging from an edge.  We never want to use that edge's vertices as targets.
 											bIsDisqualified = true;
 										}
 									}
 									else if( ensure( StartingVertexID != FVertexID::Invalid ) )
 									{
-										static TArray<FEdgeID> AdjacentEdgeIDs;
-										EditableMesh->GetVertexConnectedEdges( StartingVertexID, /* Out */ AdjacentEdgeIDs );
-										if( AdjacentEdgeIDs.Contains( PerimeterEdgeID ) )
+										if( TargetVertexID == StartingVertexID )
 										{
+											// The vertex we dragged from is disqualified as a target
 											bIsDisqualified = true;
 										}
+										else
+										{
+											// The vertices that share an edge with our starting vertex are disqualified, because we don't want
+											// to create an edge that's colinear with an existing edge
+											static TArray<FVertexID> AdjacentVertexIDs;
+											EditableMesh->GetVertexAdjacentVertices( StartingVertexID, /* Out */ AdjacentVertexIDs );
+
+											if( AdjacentVertexIDs.Contains( TargetVertexID ) )
+											{
+												bIsDisqualified = true;
+											}
+										}
 									}
-								}
 
-								// Never split an edge that's directly connected to the starting vertex
-								if( !bIsDisqualified )
-								{
-									CandidateEdgeIDs.AddUnique( PerimeterEdgeID );
+									if( !bIsDisqualified )
+									{
+										// Connect to this vertex!
+										PolygonToSplit = CandidatePolygonRef;
+										ToVertexID = TargetVertexID;
+									}
 								}
-							}
-						}
-					}
 
-					if( CandidateEdgeIDs.Num() > 0 )
-					{
-						static TArray<FMeshElement> CandidateEdgeElements;
-						CandidateEdgeElements.Reset();
-						for( const FEdgeID CandidateEdgeID : CandidateEdgeIDs )
-						{
-							FMeshElement CandidateEdgeElement;
-							CandidateEdgeElement.Component = Component;
-							CandidateEdgeElement.ElementAddress.SubMeshAddress = EditableMesh->GetSubMeshAddress();
-							CandidateEdgeElement.ElementAddress.ElementType = EEditableMeshElementType::Edge;
-							CandidateEdgeElement.ElementAddress.ElementID = CandidateEdgeID;
-
-							CandidateEdgeElements.Add( CandidateEdgeElement );
-						}
-
-						// Figure out where to split
-						FEdgeID TargetEdgeID = FEdgeID::Invalid;
-						float OtherEdgeSplit = 0.0f;
-						const bool bFoundSplit = MeshEditorMode.FindEdgeSplitUnderInteractor( ViewportInteractor, EditableMesh, CandidateEdgeElements, /* Out */ TargetEdgeID, /* Out */ OtherEdgeSplit );
-						if( bFoundSplit )
-						{
-							// Figure out which candidate polygon to split
-							for( const FPolygonRef CandidatePolygonRef : CandidatePolygonRefs )
-							{
-								static TArray<FEdgeID> PerimeterEdgeIDs;
-								EditableMesh->GetPolygonPerimeterEdges( CandidatePolygonRef, /* Out */ PerimeterEdgeIDs );
-								if( PerimeterEdgeIDs.Contains( TargetEdgeID ) )
+								// If a vertex wasn't eligible, go ahead and connect to the edge
+								if( ToVertexID == FVertexID::Invalid )
 								{
+									// Split the edge to create a new vertex that we'll connect to
 									PolygonToSplit = CandidatePolygonRef;
-									break;
-								}
-							}
-							check( PolygonToSplit != FPolygonRef::Invalid );
 
-							// Go ahead and split the target edge
-							static TArray<FVertexID> NewVertexIDs;
-							NewVertexIDs.Reset();
+									// Go ahead and split the target edge
+									static TArray<FVertexID> NewVertexIDs;
+									NewVertexIDs.Reset();
 
-							static TArray<float> SplitEdgeSplitList;
-							SplitEdgeSplitList.SetNumUninitialized( 1 );
-							SplitEdgeSplitList[ 0 ] = OtherEdgeSplit;
+									static TArray<float> SplitEdgeSplitList;
+									SplitEdgeSplitList.SetNumUninitialized( 1 );
+									SplitEdgeSplitList[ 0 ] = ImpactProgressAlongEdge;
 
-							EditableMesh->SplitEdge( TargetEdgeID, SplitEdgeSplitList, /* Out */ NewVertexIDs );
-							check( NewVertexIDs.Num() == 1 );
+									EditableMesh->SplitEdge( ClosestEdgeID, SplitEdgeSplitList, /* Out */ NewVertexIDs );
+									check( NewVertexIDs.Num() == 1 );
 
-							// Great!  We now have a vertex to connect with
-							ToVertexID = NewVertexIDs[ 0 ];
+									// Great!  We now have a vertex to connect with
+									ToVertexID = NewVertexIDs[ 0 ];
+								}
+							}
 						}
 					}
 				}
 			}
 		}
 
+
 		if( PolygonToSplit != FPolygonRef::Invalid )
 		{
 			check( ToVertexID != FVertexID::Invalid );
