Mesh Editor: Fixed polygon back-facing check when rendering hover/selection overlay

diff --git a/Engine/Plugins/Editor/MeshEditor/Source/MeshEditor/MeshEditorMode.cpp b/Engine/Plugins/Editor/MeshEditor/Source/MeshEditor/MeshEditorMode.cpp
--- a/Engine/Plugins/Editor/MeshEditor/Source/MeshEditor/MeshEditorMode.cpp
+++ b/Engine/Plugins/Editor/MeshEditor/Source/MeshEditor/MeshEditorMode.cpp
@@ -2050,9 +2050,7 @@ void FMeshEditorMode::AddPolygonToDynamicMesh( const UEditableMesh& EditableMesh
 	// @todo mesheditor: Support polygon hole contours!
 
 
-	// Find the normal of the face (expects counter-clockwise winding). Flip the sign with the component to world determinant in case the component has been mirrored.
-	// @todo mesheditor: For polygon, assuming planar here for normal.  Not right.
-	const FVector PolygonNormal = FVector::CrossProduct( ( PolygonPerimeterVertexPositions[ 2 ] - PolygonPerimeterVertexPositions[ 0 ] ), ( PolygonPerimeterVertexPositions[ 1 ] - PolygonPerimeterVertexPositions[ 0 ] ) ).GetSafeNormal();
+	const FVector PolygonNormal = ComponentToWorldMatrix.TransformVector( EditableMesh.ComputePolygonNormal( PolygonRef ) ).GetSafeNormal();
 	const float Determinant = ComponentToWorldMatrix.Determinant();
 
 	const FVector CameraForwardVectorUnnormalized = CameraToWorld.TransformVector( FVector::ForwardVector );
