Split Polygon: Avoid creating colinear and degenerate triangles when splitting polygons

diff --git a/Engine/Plugins/Editor/MeshEditor/Source/PolygonModeling/SplitPolygon.cpp b/Engine/Plugins/Editor/MeshEditor/Source/PolygonModeling/SplitPolygon.cpp
--- a/Engine/Plugins/Editor/MeshEditor/Source/PolygonModeling/SplitPolygon.cpp
+++ b/Engine/Plugins/Editor/MeshEditor/Source/PolygonModeling/SplitPolygon.cpp
@@ -205,21 +205,31 @@ void USplitPolygonCommand::ApplyDuringDrag( IMeshEditorModeEditingContract& Mesh
 									const FVector EdgeVertex0Location = ComponentToWorld.TransformPosition( EditableMesh->GetVertexAttribute( EdgeVertex0, UEditableMeshAttribute::VertexPosition(), 0 ) );
 									const FVector EdgeVertex1Location = ComponentToWorld.TransformPosition( EditableMesh->GetVertexAttribute( EdgeVertex1, UEditableMeshAttribute::VertexPosition(), 0 ) );
 
-									FVector ClosestPointOnSplitLine, ClosestPointOnEdge;
-									FMath::SegmentDistToSegmentSafe(
-										SplitStartLocation, SplitStartLocation + SplitDirection * 99999.0f,
-										EdgeVertex0Location, EdgeVertex1Location,
-										/* Out */ ClosestPointOnSplitLine,
-										/* Out */ ClosestPointOnEdge );
-
-									// Closest points should be the same if there was actually an intersection
-									if( ClosestPointOnSplitLine.Equals( ClosestPointOnEdge ) )
+									// Don't allow connecting to edges that are either degenerate or colinear to the split direction
+									const FVector EdgeDirection = ( EdgeVertex1Location - EdgeVertex0Location ).GetSafeNormal();
+									if( FMath::IsNearlyZero( EdgeDirection.SizeSquared() ) || FMath::IsNearlyEqual( FMath::Abs( FVector::DotProduct( SplitDirection, EdgeDirection ) ), 1.0f ) )
 									{
-										const float DistanceToEdgeImpact = ( ClosestPointOnEdge - SplitStartLocation ).Size();
-										if( DistanceToEdgeImpact < ClosestEdgeDistance )
+										bIsDisqualified = true;
+									}
+
+									if( !bIsDisqualified )
+									{
+										FVector ClosestPointOnSplitLine, ClosestPointOnEdge;
+										FMath::SegmentDistToSegmentSafe(
+											SplitStartLocation, SplitStartLocation + SplitDirection * 99999.0f,
+											EdgeVertex0Location, EdgeVertex1Location,
+											/* Out */ ClosestPointOnSplitLine,
+											/* Out */ ClosestPointOnEdge );
+
+										// Closest points should be the same if there was actually an intersection
+										if( ClosestPointOnSplitLine.Equals( ClosestPointOnEdge ) )
 										{
-											ClosestEdgeDistance = DistanceToEdgeImpact;
-											ClosestEdgeID = TargetEdgeID;
+											const float DistanceToEdgeImpact = ( ClosestPointOnEdge - SplitStartLocation ).Size();
+											if( DistanceToEdgeImpact < ClosestEdgeDistance )
+											{
+												ClosestEdgeDistance = DistanceToEdgeImpact;
+												ClosestEdgeID = TargetEdgeID;
+											}
 										}
 									}
 								}
@@ -283,9 +293,20 @@ void USplitPolygonCommand::ApplyDuringDrag( IMeshEditorModeEditingContract& Mesh
 
 									if( !bIsDisqualified )
 									{
-										// Connect to this vertex!
-										PolygonToSplit = CandidatePolygonRef;
-										ToVertexID = TargetVertexID;
+										// Don't allow connecting to vertices that are colinear to the split direction
+										const FVector VertexLocation = ComponentToWorld.TransformPosition( EditableMesh->GetVertexAttribute( TargetVertexID, UEditableMeshAttribute::VertexPosition(), 0 ) );
+										const FVector VertexDirection = ( VertexLocation - SplitStartLocation ).GetSafeNormal();
+										if( FMath::IsNearlyZero( VertexDirection.SizeSquared() ) || FMath::IsNearlyEqual( FMath::Abs( FVector::DotProduct( SplitDirection, VertexDirection ) ), 1.0f ) )
+										{
+											bIsDisqualified = true;
+										}
+
+										if( !bIsDisqualified )
+										{
+											// Connect to this vertex!
+											PolygonToSplit = CandidatePolygonRef;
+											ToVertexID = TargetVertexID;
+										}
 									}
 								}
 
