Split Polygon: Don't split unless pointer is close to an edge or vertex

diff --git a/Engine/Plugins/Editor/MeshEditor/Source/PolygonModeling/SplitPolygon.cpp b/Engine/Plugins/Editor/MeshEditor/Source/PolygonModeling/SplitPolygon.cpp
--- a/Engine/Plugins/Editor/MeshEditor/Source/PolygonModeling/SplitPolygon.cpp
+++ b/Engine/Plugins/Editor/MeshEditor/Source/PolygonModeling/SplitPolygon.cpp
@@ -154,182 +154,189 @@ void USplitPolygonCommand::ApplyDuringDrag( IMeshEditorModeEditingContract& Mesh
 					}
 
 
+					const FPlane PolygonPlane = EditableMesh->ComputePolygonPlane( CandidatePolygonRef ).TransformBy( ComponentToWorld.ToMatrixWithScale() );
+					const FVector LaserImpactOnPolygonPlane = FMath::LinePlaneIntersection( LaserPointerStart, LaserPointerEnd, PolygonPlane );
+
+					// @todo mesheditor splitpolygon: Ideally this would be more "fuzzy", and allow the interactor to extend beyond the range of the polygon.  But it will make figuring out which polygon to split more tricky.
+					// @todo mesheditor urgent: Can crash with "Colinear points in FMath::ComputeBaryCentric2D()"  Needs repro.
+					static TArray<int32> PerimeterVertexIndices;
+					FVector TriangleVertexWeights;
+					if( EditableMesh->ComputeBarycentricWeightForPointOnPolygon( CandidatePolygonRef, ComponentToWorld.InverseTransformPosition( LaserImpactOnPolygonPlane ), /* Out */ PerimeterVertexIndices, /* Out */ TriangleVertexWeights ) )
 					{
-						const FPlane PolygonPlane = EditableMesh->ComputePolygonPlane( CandidatePolygonRef ).TransformBy( ComponentToWorld.ToMatrixWithScale() );
-	
-						const FVector LaserImpactOnPolygonPlane = FMath::LinePlaneIntersection( LaserPointerStart, LaserPointerEnd, PolygonPlane );
-
-						// @todo mesheditor splitpolygon: Ideally this would be more "fuzzy", and allow the interactor to extend beyond the range of the polygon.  But it will make figuring out which polygon to split more tricky.
-						// @todo mesheditor urgent: Can crash with "Colinear points in FMath::ComputeBaryCentric2D()"  Needs repro.
-						static TArray<int32> PerimeterVertexIndices;
-						FVector TriangleVertexWeights;
-						if( EditableMesh->ComputeBarycentricWeightForPointOnPolygon( CandidatePolygonRef, ComponentToWorld.InverseTransformPosition( LaserImpactOnPolygonPlane ), /* Out */ PerimeterVertexIndices, /* Out */ TriangleVertexWeights ) )
-						{
-							const FVector SplitDirection = ( LaserImpactOnPolygonPlane - SplitStartLocation ).GetSafeNormal();
+						const FVector SplitDirection = ( LaserImpactOnPolygonPlane - SplitStartLocation ).GetSafeNormal();
 
-							// Trace out within the polygon to figure out where the split should connect to
-							float ClosestEdgeDistance = TNumericLimits<float>::Max();
-							FEdgeID ClosestEdgeID = FEdgeID::Invalid;
+						// Trace out within the polygon to figure out where the split should connect to
+						float ClosestEdgeDistance = TNumericLimits<float>::Max();
+						FEdgeID ClosestEdgeID = FEdgeID::Invalid;
 
-							static TArray<FEdgeID> PolygonPerimeterEdgeIDs;
-							EditableMesh->GetPolygonPerimeterEdges( CandidatePolygonRef, /* Out */ PolygonPerimeterEdgeIDs );
+						static TArray<FEdgeID> PolygonPerimeterEdgeIDs;
+						EditableMesh->GetPolygonPerimeterEdges( CandidatePolygonRef, /* Out */ PolygonPerimeterEdgeIDs );
 
-							for( const FEdgeID TargetEdgeID : PolygonPerimeterEdgeIDs )
-							{
-								bool bIsDisqualified = false;
+						for( const FEdgeID TargetEdgeID : PolygonPerimeterEdgeIDs )
+						{
+							bool bIsDisqualified = false;
 
-								if( StartingEdgeID != FEdgeID::Invalid )
+							if( StartingEdgeID != FEdgeID::Invalid )
+							{
+								if( StartingEdgeID == TargetEdgeID )
 								{
-									if( StartingEdgeID == TargetEdgeID )
-									{
-										// The edge we dragged from is disqualified as a target
-										bIsDisqualified = true;
-									}
+									// The edge we dragged from is disqualified as a target
+									bIsDisqualified = true;
 								}
-								else if( ensure( StartingVertexID != FVertexID::Invalid ) )
+							}
+							else if( ensure( StartingVertexID != FVertexID::Invalid ) )
+							{
+								static TArray<FEdgeID> AdjacentEdgeIDs;
+								EditableMesh->GetVertexConnectedEdges( StartingVertexID, /* Out */ AdjacentEdgeIDs );
+								if( AdjacentEdgeIDs.Contains( TargetEdgeID ) )
 								{
-									static TArray<FEdgeID> AdjacentEdgeIDs;
-									EditableMesh->GetVertexConnectedEdges( StartingVertexID, /* Out */ AdjacentEdgeIDs );
-									if( AdjacentEdgeIDs.Contains( TargetEdgeID ) )
-									{
-										// Never split an edge that's directly connected to the starting vertex
-										bIsDisqualified = true;
-									}
+									// Never split an edge that's directly connected to the starting vertex
+									bIsDisqualified = true;
 								}
+							}
 
-								if( !bIsDisqualified )
-								{
-									FVertexID EdgeVertex0, EdgeVertex1;
-									EditableMesh->GetEdgeVertices( TargetEdgeID, /* Out */ EdgeVertex0, /* Out */ EdgeVertex1 );
+							FVertexID EdgeVertex0, EdgeVertex1;
+							EditableMesh->GetEdgeVertices( TargetEdgeID, /* Out */ EdgeVertex0, /* Out */ EdgeVertex1 );
+							const FVector EdgeVertex0Location = ComponentToWorld.TransformPosition( EditableMesh->GetVertexAttribute( EdgeVertex0, UEditableMeshAttribute::VertexPosition(), 0 ) );
+							const FVector EdgeVertex1Location = ComponentToWorld.TransformPosition( EditableMesh->GetVertexAttribute( EdgeVertex1, UEditableMeshAttribute::VertexPosition(), 0 ) );
 
-									const FVector EdgeVertex0Location = ComponentToWorld.TransformPosition( EditableMesh->GetVertexAttribute( EdgeVertex0, UEditableMeshAttribute::VertexPosition(), 0 ) );
-									const FVector EdgeVertex1Location = ComponentToWorld.TransformPosition( EditableMesh->GetVertexAttribute( EdgeVertex1, UEditableMeshAttribute::VertexPosition(), 0 ) );
+							if( !bIsDisqualified )
+							{
+								// Don't bother trying to split unless the laser impact point is reasonably close to the target edge.  Otherwise it just feels
+								// bad, because our split direction will be unstable as the interactor moves along the polygon
+								const float Distance = FMath::PointDistToSegment( LaserImpactOnPolygonPlane, EdgeVertex0Location, EdgeVertex1Location );
+								if( Distance > 20.0f )	// @todo mesheditor tweak: Should be based on polygon area or edge size, and probably scale with distance like other fuzzy tests.  Also, consider making this a distance *from* the impact point instead.
+								{
+									bIsDisqualified = true;
+								}
+							}
 
-									// Don't allow connecting to edges that are either degenerate or colinear to the split direction
-									const FVector EdgeDirection = ( EdgeVertex1Location - EdgeVertex0Location ).GetSafeNormal();
-									if( FMath::IsNearlyZero( EdgeDirection.SizeSquared() ) || FMath::IsNearlyEqual( FMath::Abs( FVector::DotProduct( SplitDirection, EdgeDirection ) ), 1.0f ) )
-									{
-										bIsDisqualified = true;
-									}
+							if( !bIsDisqualified )
+							{
+								// Don't allow connecting to edges that are either degenerate or colinear to the split direction
+								const FVector EdgeDirection = ( EdgeVertex1Location - EdgeVertex0Location ).GetSafeNormal();
+								if( FMath::IsNearlyZero( EdgeDirection.SizeSquared() ) || FMath::IsNearlyEqual( FMath::Abs( FVector::DotProduct( SplitDirection, EdgeDirection ) ), 1.0f ) )
+								{
+									bIsDisqualified = true;
+								}
 
-									if( !bIsDisqualified )
+								if( !bIsDisqualified )
+								{
+									FVector ClosestPointOnSplitLine, ClosestPointOnEdge;
+									FMath::SegmentDistToSegmentSafe(
+										SplitStartLocation, SplitStartLocation + SplitDirection * 99999.0f,
+										EdgeVertex0Location, EdgeVertex1Location,
+										/* Out */ ClosestPointOnSplitLine,
+										/* Out */ ClosestPointOnEdge );
+
+									// Closest points should be the same if there was actually an intersection
+									if( ClosestPointOnSplitLine.Equals( ClosestPointOnEdge ) )
 									{
-										FVector ClosestPointOnSplitLine, ClosestPointOnEdge;
-										FMath::SegmentDistToSegmentSafe(
-											SplitStartLocation, SplitStartLocation + SplitDirection * 99999.0f,
-											EdgeVertex0Location, EdgeVertex1Location,
-											/* Out */ ClosestPointOnSplitLine,
-											/* Out */ ClosestPointOnEdge );
-
-										// Closest points should be the same if there was actually an intersection
-										if( ClosestPointOnSplitLine.Equals( ClosestPointOnEdge ) )
+										const float DistanceToEdgeImpact = ( ClosestPointOnEdge - SplitStartLocation ).Size();
+										if( DistanceToEdgeImpact < ClosestEdgeDistance )
 										{
-											const float DistanceToEdgeImpact = ( ClosestPointOnEdge - SplitStartLocation ).Size();
-											if( DistanceToEdgeImpact < ClosestEdgeDistance )
-											{
-												ClosestEdgeDistance = DistanceToEdgeImpact;
-												ClosestEdgeID = TargetEdgeID;
-											}
+											ClosestEdgeDistance = DistanceToEdgeImpact;
+											ClosestEdgeID = TargetEdgeID;
 										}
 									}
 								}
 							}
+						}
 
-							if( ClosestEdgeID != FEdgeID::Invalid )
-							{
-								FVertexID EdgeVertex0, EdgeVertex1;
-								EditableMesh->GetEdgeVertices( ClosestEdgeID, /* Out */ EdgeVertex0, /* Out */ EdgeVertex1 );
+						if( ClosestEdgeID != FEdgeID::Invalid )
+						{
+							FVertexID EdgeVertex0, EdgeVertex1;
+							EditableMesh->GetEdgeVertices( ClosestEdgeID, /* Out */ EdgeVertex0, /* Out */ EdgeVertex1 );
 
-								const FVector EdgeVertex0Location = ComponentToWorld.TransformPosition( EditableMesh->GetVertexAttribute( EdgeVertex0, UEditableMeshAttribute::VertexPosition(), 0 ) );
-								const FVector EdgeVertex1Location = ComponentToWorld.TransformPosition( EditableMesh->GetVertexAttribute( EdgeVertex1, UEditableMeshAttribute::VertexPosition(), 0 ) );
+							const FVector EdgeVertex0Location = ComponentToWorld.TransformPosition( EditableMesh->GetVertexAttribute( EdgeVertex0, UEditableMeshAttribute::VertexPosition(), 0 ) );
+							const FVector EdgeVertex1Location = ComponentToWorld.TransformPosition( EditableMesh->GetVertexAttribute( EdgeVertex1, UEditableMeshAttribute::VertexPosition(), 0 ) );
 
-								const FVector ImpactOnEdge = SplitStartLocation + SplitDirection * ClosestEdgeDistance;
+							const FVector ImpactOnEdge = SplitStartLocation + SplitDirection * ClosestEdgeDistance;
 
-								const float EdgeLength = ( EdgeVertex1Location - EdgeVertex0Location ).Size();
-								const float ImpactProgressAlongEdge = ( ImpactOnEdge - EdgeVertex0Location ).Size() / EdgeLength;
+							const float EdgeLength = ( EdgeVertex1Location - EdgeVertex0Location ).Size();
+							const float ImpactProgressAlongEdge = ( ImpactOnEdge - EdgeVertex0Location ).Size() / EdgeLength;
 
-								// If we're really close to one side or the other of an edge, try to connect with an existing vertex instead
-								const float EdgeProgressVertexSnapThreshold = 0.075f;	// @todo mesheditor splitpolygon: Should be actual 'fuzzy' distance consistent with MeshEditorMode, in world/screen units, not percentage of edge progress
-								const bool bIsCloseToEdgeVertex0 = ImpactProgressAlongEdge < EdgeProgressVertexSnapThreshold;
-								const bool bIsCloseToEdgeVertex1 = ImpactProgressAlongEdge > 1.0f - EdgeProgressVertexSnapThreshold;
-								if( bIsCloseToEdgeVertex0 || bIsCloseToEdgeVertex1 )
-								{
-									// Prefer connecting to a vertex
-									const FVertexID TargetVertexID = bIsCloseToEdgeVertex0 ? EdgeVertex0 : EdgeVertex1;
+							// If we're really close to one side or the other of an edge, try to connect with an existing vertex instead
+							const float EdgeProgressVertexSnapThreshold = 0.075f;	// @todo mesheditor splitpolygon: Should be actual 'fuzzy' distance consistent with MeshEditorMode, in world/screen units, not percentage of edge progress
+							const bool bIsCloseToEdgeVertex0 = ImpactProgressAlongEdge < EdgeProgressVertexSnapThreshold;
+							const bool bIsCloseToEdgeVertex1 = ImpactProgressAlongEdge > 1.0f - EdgeProgressVertexSnapThreshold;
+							if( bIsCloseToEdgeVertex0 || bIsCloseToEdgeVertex1 )
+							{
+								// Prefer connecting to a vertex
+								const FVertexID TargetVertexID = bIsCloseToEdgeVertex0 ? EdgeVertex0 : EdgeVertex1;
+
+								bool bIsDisqualified = false;
 
-									bool bIsDisqualified = false;
+								if( StartingEdgeID != FEdgeID::Invalid )
+								{
+									FVertexID StartingEdgeVertex0, StartingEdgeVertex1;
+									EditableMesh->GetEdgeVertices( StartingEdgeID, /* Out */ StartingEdgeVertex0, /* Out */ StartingEdgeVertex1 );
 
-									if( StartingEdgeID != FEdgeID::Invalid )
+									if( TargetVertexID == StartingEdgeVertex0 || TargetVertexID == StartingEdgeVertex1 )
 									{
-										FVertexID StartingEdgeVertex0, StartingEdgeVertex1;
-										EditableMesh->GetEdgeVertices( StartingEdgeID, /* Out */ StartingEdgeVertex0, /* Out */ StartingEdgeVertex1 );
+										// We're dragging from an edge.  We never want to use that edge's vertices as targets.
+										bIsDisqualified = true;
+									}
+								}
+								else if( ensure( StartingVertexID != FVertexID::Invalid ) )
+								{
+									if( TargetVertexID == StartingVertexID )
+									{
+										// The vertex we dragged from is disqualified as a target
+										bIsDisqualified = true;
+									}
+									else
+									{
+										// The vertices that share an edge with our starting vertex are disqualified, because we don't want
+										// to create an edge that's colinear with an existing edge
+										static TArray<FVertexID> AdjacentVertexIDs;
+										EditableMesh->GetVertexAdjacentVertices( StartingVertexID, /* Out */ AdjacentVertexIDs );
 
-										if( TargetVertexID == StartingEdgeVertex0 || TargetVertexID == StartingEdgeVertex1 )
+										if( AdjacentVertexIDs.Contains( TargetVertexID ) )
 										{
-											// We're dragging from an edge.  We never want to use that edge's vertices as targets.
 											bIsDisqualified = true;
 										}
 									}
-									else if( ensure( StartingVertexID != FVertexID::Invalid ) )
+								}
+
+								if( !bIsDisqualified )
+								{
+									// Don't allow connecting to vertices that are colinear to the split direction
+									const FVector VertexLocation = ComponentToWorld.TransformPosition( EditableMesh->GetVertexAttribute( TargetVertexID, UEditableMeshAttribute::VertexPosition(), 0 ) );
+									const FVector VertexDirection = ( VertexLocation - SplitStartLocation ).GetSafeNormal();
+									if( FMath::IsNearlyZero( VertexDirection.SizeSquared() ) || FMath::IsNearlyEqual( FMath::Abs( FVector::DotProduct( SplitDirection, VertexDirection ) ), 1.0f ) )
 									{
-										if( TargetVertexID == StartingVertexID )
-										{
-											// The vertex we dragged from is disqualified as a target
-											bIsDisqualified = true;
-										}
-										else
-										{
-											// The vertices that share an edge with our starting vertex are disqualified, because we don't want
-											// to create an edge that's colinear with an existing edge
-											static TArray<FVertexID> AdjacentVertexIDs;
-											EditableMesh->GetVertexAdjacentVertices( StartingVertexID, /* Out */ AdjacentVertexIDs );
-
-											if( AdjacentVertexIDs.Contains( TargetVertexID ) )
-											{
-												bIsDisqualified = true;
-											}
-										}
+										bIsDisqualified = true;
 									}
 
 									if( !bIsDisqualified )
 									{
-										// Don't allow connecting to vertices that are colinear to the split direction
-										const FVector VertexLocation = ComponentToWorld.TransformPosition( EditableMesh->GetVertexAttribute( TargetVertexID, UEditableMeshAttribute::VertexPosition(), 0 ) );
-										const FVector VertexDirection = ( VertexLocation - SplitStartLocation ).GetSafeNormal();
-										if( FMath::IsNearlyZero( VertexDirection.SizeSquared() ) || FMath::IsNearlyEqual( FMath::Abs( FVector::DotProduct( SplitDirection, VertexDirection ) ), 1.0f ) )
-										{
-											bIsDisqualified = true;
-										}
-
-										if( !bIsDisqualified )
-										{
-											// Connect to this vertex!
-											PolygonToSplit = CandidatePolygonRef;
-											ToVertexID = TargetVertexID;
-										}
+										// Connect to this vertex!
+										PolygonToSplit = CandidatePolygonRef;
+										ToVertexID = TargetVertexID;
 									}
 								}
+							}
 
-								// If a vertex wasn't eligible, go ahead and connect to the edge
-								if( ToVertexID == FVertexID::Invalid )
-								{
-									// Split the edge to create a new vertex that we'll connect to
-									PolygonToSplit = CandidatePolygonRef;
+							// If a vertex wasn't eligible, go ahead and connect to the edge
+							if( ToVertexID == FVertexID::Invalid )
+							{
+								// Split the edge to create a new vertex that we'll connect to
+								PolygonToSplit = CandidatePolygonRef;
 
-									// Go ahead and split the target edge
-									static TArray<FVertexID> NewVertexIDs;
-									NewVertexIDs.Reset();
+								// Go ahead and split the target edge
+								static TArray<FVertexID> NewVertexIDs;
+								NewVertexIDs.Reset();
 
-									static TArray<float> SplitEdgeSplitList;
-									SplitEdgeSplitList.SetNumUninitialized( 1 );
-									SplitEdgeSplitList[ 0 ] = ImpactProgressAlongEdge;
+								static TArray<float> SplitEdgeSplitList;
+								SplitEdgeSplitList.SetNumUninitialized( 1 );
+								SplitEdgeSplitList[ 0 ] = ImpactProgressAlongEdge;
 
-									EditableMesh->SplitEdge( ClosestEdgeID, SplitEdgeSplitList, /* Out */ NewVertexIDs );
-									check( NewVertexIDs.Num() == 1 );
+								EditableMesh->SplitEdge( ClosestEdgeID, SplitEdgeSplitList, /* Out */ NewVertexIDs );
+								check( NewVertexIDs.Num() == 1 );
 
-									// Great!  We now have a vertex to connect with
-									ToVertexID = NewVertexIDs[ 0 ];
-								}
+								// Great!  We now have a vertex to connect with
+								ToVertexID = NewVertexIDs[ 0 ];
 							}
 						}
 					}
