Copying //Tasks/UE4/Dev-Geometry-Refactor to Dev-Geometry (//UE4/Dev-Geometry)

CL 3365379:
Beginnings of 'uber mesh'. Moved mesh description across to UEditableMesh.  Renamed Rendering Vertices to Vertex Instances, they will form a fundamental part of the mesh description, in representing shared vertices.

CL 3365385:
Renamed Sections to PolygonGroups. This prevents confusion between editable mesh sections and rendering sections (from the static mesh). It also means we can potentially open up the concept of 'polygon groups' to represent more than static mesh sections if we wish to.

CL 3365387:
Removed FPolygonRef and flattened Polygons container. Added new methods to act on the entire Polygons container. Added RenderingPolygons, RenderingTriangles, RenderingPolygonGroups containers which will eventually live on the adapter, and which are solely involved in maintaining the render buffers. Fixed Quadrangulate so that it doesn't merge across polygon groups.

CL 3365390:
Started to remove virtuals from EditableMesh. First separation of ubermesh intrinsic properties from implementation details in EditableStaticMesh.

CL 3396167:
Picking out generic functionality from the _Internal methods, moving them into UEditableMesh code, ready to remove the inheritance entirely. Implemented adapter interface (currently living within UEditableStaticMesh) and calling into it from UEditableMesh.

Big revamp to handling of rendering vertices; now exist as vertex instances, which are dynamically shared / split, depending on whether colors / textures are identical. Normal/tangent is handled separately as it is autogenerated. Whether neighboring edges are hard or not determines whether a vertex will be automatically instanced, regardless of other attributes. We infer whether a vertex is 'hard' from whether the mesh is split by hard edges at that point (done by tracing adjacent polygons through only soft edges); hard vertices are always instanced for that poly. This simplifies the autogeneration of normals, as it now simply iterates through stale vertex instances, applying the weighted contribution from each connected polygon.  Polygon normals are now cached as a read-only attribute, to make this process even quicker.

Rewrite of the basic primitive operations (Create/Delete polygon / edges / vertices / vertex instances). Added a new primitive operation to substitute vertex instances on a polygon contour.  Setting edge hardness attribute now potentially instances or deletes vertex instances, as does setting vertex instance attributes.

CL 3396229:
Intermediate check-in (phase 2).

Removed a whole load of dead code. Broke functionality, until InitEditableStaticMesh is rewritten to use the public UEditableMesh API.

CL 3401798:
Intermediate check-in (phase 3).

Implemented a simple alternative to tangent basis generation as MikkTSpace doesn't quite fit our model currently (but with a view to revisiting it in the future; nothing has been removed for now).
Fixed edge hardness setting.
Added missing hash function for FVertexInstanceID.

CL 3403360:
Fixed some bugs in tangent/normal handling. Now correctly includes contributions from adjacent polygons in the same smoothing group but different vertex instance when calculating a vertex instance normal.

CL 3411429:
Fixed a bunch of bugs:
- Fixed issue where adjacent polygons the other side of a UV seam were not contributing to vertex instance normal calculation.
- Fixed various issues where out arrays weren't being reset, and mesh change actions were not being called after setting up the input array.
- Various changes to converting static meshes to editable meshes.
- Removed old ContourVertices arrays from polygon contours; now polygons exclusively use an array of vertex instances.
- Updated natvis to parse new element ID types.

CL 3418022:
More bug fixes:
- Sorted an issue with undo where vertex instance merging / splitting was happening twice, once due to the parent operation, and another time due to the revert transactions which are created.
- Added options to delete orphaned vertex instances in DeletePolygons and other operations which remove vertex instances.
- Fixed issue with SplitEdge due to subtle bug where connected edges were duplicated in a vertex, instead of being added uniquely.
- Changed order that undo items are added to the stack so that the revert operation works smoothly.

CL 3427887:
Fixed issue which occurs in certain circumstances with Split Edge.
Fixed regression which was causing various operations to fail when setting edge hardness.

CL 3430665:
Specify polygon edge hardness policy when creating a polygon. Simplified some existing uses of CreatePolygons. Changed behavior to match original behavior so that only *new* edge IDs are returned by CreatePolygons.

CL 3430927:
Corrected normal/tangent generation for polygons and adjacent polygons.

CL 3431134:
Fixed small issue with remove vertex in certain circumstances.

CL 3446929:
- Added a triangle list to the FMeshPolygon struct, making triangulation a first-class operation instead of an implementation detail.
- Repurposed UEditableStaticMesh as an editable mesh adapter, in a separate class hierarchy.
- Tidied up interfaces
- Moved serialization, compact/uncompact into UEditableMesh
- Currently broken: instance changes. This is pending a rethink on how we scope objects in the assets.

CL 3446954:
Renamed EditableStaticMesh.* to EditableStaticMeshAdapter.* to reflect its new role and place in the class hierarchy.

CL 3449133:
Fixed issue with splitting vertices due to the wrong polygon ID being used by GetPolygonsInSameSmoothingGroup. Also fixed an error where the list being iterated was potentially being mutated by the split operation.

CL 3449362:
Fixed ChangePolygonsVertexInstances to fix up the triangle list after replacing vertex instances. Not sure where this code went!

CL 3458823:
- Deferred normals calculation until the end of the transaction for efficiency
- Fixed a bug in inserting edge loops
- Fixed a bug in splitting vertex instances when hardening an edge
- Moved verbose logging into a LogEditableMesh category.

CL 3461128:
Included a bunch of required header files (IWYU).

CL 3464397:
Change to handling of triangulation and normal/tangent generation. Now this is done as a deferred process during EndTransaction. This achieves three goals:
 1) Polygons are not repeatedly triangulated unnecessarily as a consequence of suboperations.
 2) We are able to calculate and cache the polygon tangent basis upfront, which is then used for every connected vertex instance tangent basis - this is quickerthan repeatedly calculating it each time it's needed.
 3) We no longer need to treat triangulation or vertex generation as a special FChange. Instead it is rolled into the FStartOrEndTransaction change. Hence we can remove the awkward code which was triangulating polygons before or after other actions, depending on whether undo was being processed.
Note: seem to have introduced a bug which is getting the tangents subtly wrong in certain cases; need to investigate.

CL 3465738:
Fixed regression introduced in tangent calculations. Results now match previous version. Stripped out redundant code, and removed superfluous logging. When building the editable mesh, the tangent basis for all polygons is now precached.

CL 3465744:
Removed old debugging code.

CL 3472569:
Reinstated Compact (currently after every topology-changing transaction; will make it much more infrequent when I'm sure it works perfectly). Made some changes to fixing up element IDs. Added a new method in the Adapter interface to reindex elements on compact/uncompact.

CL 3476078:
Fixed issue when loading edited meshes. They now must always have their render mesh rebuilt on PostLoad. (This seems to arise from the way the circular dependency between Engine and MeshEditingRuntime has been resolved).

CL 3476104:
Reinstated instanced editable meshes.

CL 3484401:
Added missing file (EditableMeshAdapter.h)

CL 3484699:
Added new adapter method OnCreateEmptyVertexRange, now needed again by extrude (I removed this action, but Mike's new extrude requires it)

diff --git a/Engine/Plugins/Editor/MeshEditor/Source/MeshEditor/MeshEditorMode.cpp b/Engine/Plugins/Editor/MeshEditor/Source/MeshEditor/MeshEditorMode.cpp
--- a/Engine/Plugins/Editor/MeshEditor/Source/MeshEditor/MeshEditorMode.cpp
+++ b/Engine/Plugins/Editor/MeshEditor/Source/MeshEditor/MeshEditorMode.cpp
@@ -38,6 +38,10 @@
 #include "Materials/Material.h"
 #include "Materials/MaterialInterface.h"
 
+#include "Misc/FeedbackContext.h"
+#include "Framework/Application/SlateApplication.h"
+#include "Materials/Material.h"
+#include "Components/PrimitiveComponent.h"
 
 #define LOCTEXT_NAMESPACE "MeshEditorMode"
 
@@ -272,11 +276,11 @@ TUniquePtr<FChange> FMeshEditorMode::FSetElementSelectionModeChange::Execute( UO
 					else if( ElementAddress.ElementType == EEditableMeshElementType::Polygon )
 					{
 						// Select vertices of the selected polygon
-						const FPolygonRef PolygonRef = FPolygonRef( ElementAddress.SectionID, FPolygonID( ElementAddress.ElementID ) );
-						const int32 PolygonVertexCount = EditableMesh->GetPolygonPerimeterVertexCount( PolygonRef );
+						const FPolygonID PolygonID = FPolygonID( ElementAddress.ElementID );
+						const int32 PolygonVertexCount = EditableMesh->GetPolygonPerimeterVertexCount( PolygonID );
 						for( int32 Index = 0; Index < PolygonVertexCount; ++Index )
 						{
-							ElementsToSelect.Emplace( Component, SubMeshAddress, EditableMesh->GetPolygonPerimeterVertex( PolygonRef, Index ) );
+							ElementsToSelect.Emplace( Component, SubMeshAddress, EditableMesh->GetPolygonPerimeterVertex( PolygonID, Index ) );
 						}
 					}
 				}
@@ -295,12 +299,12 @@ TUniquePtr<FChange> FMeshEditorMode::FSetElementSelectionModeChange::Execute( UO
 					else if( ElementAddress.ElementType == EEditableMeshElementType::Polygon )
 					{
 						// Select edges forming the selected polygon
-						const FPolygonRef PolygonRef = FPolygonRef( ElementAddress.SectionID, FPolygonID( ElementAddress.ElementID ) );
-						const int32 PolygonEdgeCount = EditableMesh->GetPolygonPerimeterEdgeCount( PolygonRef );
+						const FPolygonID PolygonID = FPolygonID( ElementAddress.ElementID );
+						const int32 PolygonEdgeCount = EditableMesh->GetPolygonPerimeterEdgeCount( PolygonID );
 						for( int32 Index = 0; Index < PolygonEdgeCount; ++Index )
 						{
 							bool bEdgeWindingIsReversedForPolygon;
-							ElementsToSelect.Emplace( Component, SubMeshAddress, EditableMesh->GetPolygonPerimeterEdge( PolygonRef, Index, bEdgeWindingIsReversedForPolygon ) );
+							ElementsToSelect.Emplace( Component, SubMeshAddress, EditableMesh->GetPolygonPerimeterEdge( PolygonID, Index, bEdgeWindingIsReversedForPolygon ) );
 						}
 					}
 				}
@@ -310,9 +314,9 @@ TUniquePtr<FChange> FMeshEditorMode::FSetElementSelectionModeChange::Execute( UO
 					{
 						// Select all polygons containing the selected vertex
 						// @todo mesheditor: is that reasonable? Should it only select a polygon which has all its vertices selected?
-						TArray<FPolygonRef> VertexConnectedPolygons;
+						TArray<FPolygonID> VertexConnectedPolygons;
 						EditableMesh->GetVertexConnectedPolygons( FVertexID( ElementAddress.ElementID ), VertexConnectedPolygons );
-						for( FPolygonRef VertexConnectedPolygon : VertexConnectedPolygons )
+						for( const FPolygonID VertexConnectedPolygon : VertexConnectedPolygons )
 						{
 							ElementsToSelect.Emplace( Component, SubMeshAddress, VertexConnectedPolygon );
 						}
@@ -857,9 +861,11 @@ UEditableMesh* FMeshEditorMode::FindOrCreateEditableMesh( UPrimitiveComponent& C
 
 			if( GetDefault<UMeshEditorSettings>()->bAutoQuadrangulate )
 			{
-				static TArray<FPolygonRef> NewPolygonsRefs;
 				EditableMesh->StartModification( EMeshModificationType::Final, EMeshTopologyChange::TopologyChange );
-				EditableMesh->QuadrangulateMesh( NewPolygonsRefs );
+
+				static TArray<FPolygonID> NewPolygonIDs;
+				EditableMesh->QuadrangulateMesh( NewPolygonIDs );
+
 				EditableMesh->EndModification();
 			}
 
@@ -889,15 +895,18 @@ void FMeshEditorMode::RollbackPreviewChanges()
 		// to skip certain parts of the update (such as subdivision geometry refresh).  This should get us better performance!
 
 		// @todo mesheditor debug
-		//GWarn->Logf( TEXT( "---------- Rolling Back Preview Change (Object:%s) ----------" ), *Object->GetName() );
-		//PreviewRevertChange->PrintToLog( *GWarn );
+//		GWarn->Logf( TEXT( "---------- Rolling Back Preview Change (Object:%s) ----------" ), *Object->GetName() );
+//		PreviewRevertChange->PrintToLog( *GWarn );
+//		GWarn->Logf( TEXT( "---------- End (Object:%s) ----------" ), *Object->GetName() );
 
+		UE_LOG( LogEditableMesh, Verbose, TEXT( "------- ROLLING BACK PREVIEW CHANGE -------" ) );
 		TUniquePtr<FChange> UnusedChangeToUndoRevert = MoveTemp( PreviewRevertChange->Execute( Object ) );
+		UE_LOG( LogEditableMesh, Verbose, TEXT( "------- END ROLL BACK PREVIEW CHANGE -------" ) );
 
 		// @todo mesheditor debug
-		//GWarn->Logf( TEXT( "-----(Here's what the Redo looks like)-----" ) );
-		//UnusedChangeToUndoRevert->PrintToLog( *GWarn );
-		//GWarn->Logf( TEXT( "-------------------------------------------" ) );
+//		GWarn->Logf( TEXT( "-----(Here's what the Redo looks like)-----" ) );
+//		UnusedChangeToUndoRevert->PrintToLog( *GWarn );
+//		GWarn->Logf( TEXT( "-------------------------------------------" ) );
 	}
 	PreviewRevertChanges.Reset();
 }
@@ -1230,7 +1239,7 @@ void FMeshEditorMode::GetSelectedMeshesAndPolygonsPerimeterEdges( TMap<UEditable
 			}
 
 			static TArray<FEdgeID> PerimeterEdgeIDs;
-			EditableMesh->GetPolygonPerimeterEdges( FPolygonRef( FSectionID( PolygonElement.ElementAddress.SectionID ), FPolygonID( PolygonElement.ElementAddress.ElementID ) ), /* Out */ PerimeterEdgeIDs );
+			EditableMesh->GetPolygonPerimeterEdges( FPolygonID( PolygonElement.ElementAddress.ElementID ), /* Out */ PerimeterEdgeIDs );
 
 			for( const FEdgeID PerimeterEdgeID : PerimeterEdgeIDs )
 			{
@@ -1335,6 +1344,7 @@ bool FMeshEditorMode::DeleteSelectedMeshElement()
 		{
 			const bool bDeleteOrphanedEdges = true;
 			const bool bDeleteOrphanedVertices = true;
+			const bool bDeleteOrphanedVertexInstances = true;
 			const bool bDeleteEmptySections = true;
 
 			// If we deleted the same polygon on multiple selected instances of the same mesh, the polygon could already have been deleted
@@ -1347,6 +1357,7 @@ bool FMeshEditorMode::DeleteSelectedMeshElement()
 						FVertexID( MeshElementToDelete.ElementAddress.ElementID ),
 						bDeleteOrphanedEdges,
 						bDeleteOrphanedVertices,
+						bDeleteOrphanedVertexInstances,
 						bDeleteEmptySections );
 
 				}
@@ -1356,14 +1367,20 @@ bool FMeshEditorMode::DeleteSelectedMeshElement()
 						FEdgeID( MeshElementToDelete.ElementAddress.ElementID ),
 						bDeleteOrphanedEdges,
 						bDeleteOrphanedVertices,
+						bDeleteOrphanedVertexInstances,
 						bDeleteEmptySections );
 				}
 				else if( MeshElementToDelete.ElementAddress.ElementType == EEditableMeshElementType::Polygon )
 				{
-					static TArray<FPolygonRef> PolygonRefsToDelete;
-					PolygonRefsToDelete.Reset();
-					PolygonRefsToDelete.Add( FPolygonRef( MeshElementToDelete.ElementAddress.SectionID, FPolygonID( MeshElementToDelete.ElementAddress.ElementID ) ) );
-					EditableMesh->DeletePolygons( PolygonRefsToDelete, bDeleteOrphanedEdges, bDeleteOrphanedVertices, bDeleteEmptySections );
+					static TArray<FPolygonID> PolygonIDsToDelete;
+					PolygonIDsToDelete.Reset();
+					PolygonIDsToDelete.Add( FPolygonID( MeshElementToDelete.ElementAddress.ElementID ) );
+					EditableMesh->DeletePolygons( 
+						PolygonIDsToDelete,
+						bDeleteOrphanedEdges,
+						bDeleteOrphanedVertices,
+						bDeleteOrphanedVertexInstances,
+						bDeleteEmptySections );
 				}
 			}
 		}
@@ -1489,12 +1506,12 @@ void FMeshEditorMode::FrameSelectedElements( FEditorViewportClient* ViewportClie
 					const UPrimitiveComponent* Component = PolygonElement.Component.Get();
 					if( Component )
 					{
-						const FPolygonRef PolygonRef( PolygonElement.ElementAddress.SectionID, FPolygonID( PolygonElement.ElementAddress.ElementID ) );
+						const FPolygonID PolygonID( PolygonElement.ElementAddress.ElementID );
 
-						const int32 PolygonVertexCount = EditableMesh->GetPolygonPerimeterVertexCount( PolygonRef );
+						const int32 PolygonVertexCount = EditableMesh->GetPolygonPerimeterVertexCount( PolygonID );
 						for( int32 PolygonVertexIndex = 0; PolygonVertexIndex < PolygonVertexCount; ++PolygonVertexIndex )
 						{
-							const FVector VertexPosition = EditableMesh->GetPolygonPerimeterVertexAttribute( PolygonRef, PolygonVertexIndex, UEditableMeshAttribute::VertexPosition(), 0 );
+							const FVector VertexPosition = EditableMesh->GetPolygonPerimeterVertexAttribute( PolygonID, PolygonVertexIndex, UEditableMeshAttribute::VertexPosition(), 0 );
 							BoundingBox += Component->GetComponentTransform().TransformPosition( VertexPosition );
 						}
 					}
@@ -1675,13 +1692,13 @@ bool FMeshEditorMode::FlipSelectedPolygons()
 
 		EditableMesh->StartModification( EMeshModificationType::Final, EMeshTopologyChange::TopologyChange );
 
-		static TArray<FPolygonRef> PolygonsToFlip;
+		static TArray<FPolygonID> PolygonsToFlip;
 		PolygonsToFlip.Reset();
 
 		for( const FMeshElement& PolygonElement : MeshAndPolygons.Value )
 		{
-			const FPolygonRef PolygonRef( PolygonElement.ElementAddress.SectionID, FPolygonID( PolygonElement.ElementAddress.ElementID ) );
-			PolygonsToFlip.Add( PolygonRef );
+			const FPolygonID PolygonID( PolygonElement.ElementAddress.ElementID );
+			PolygonsToFlip.Add( PolygonID );
 		}
 
 		EditableMesh->FlipPolygons( PolygonsToFlip );
@@ -1740,19 +1757,19 @@ bool FMeshEditorMode::TriangulateSelectedPolygons()
 
 		EditableMesh->StartModification( EMeshModificationType::Final, EMeshTopologyChange::TopologyChange );
 
-		static TArray<FPolygonRef> PolygonsToTriangulate;
+		static TArray<FPolygonID> PolygonsToTriangulate;
 		PolygonsToTriangulate.Reset( MeshAndPolygons.Value.Num() );
 
 		for( const FMeshElement& PolygonElement : MeshAndPolygons.Value )
 		{
-			const FPolygonRef PolygonRef( PolygonElement.ElementAddress.SectionID, FPolygonID( PolygonElement.ElementAddress.ElementID ) );
-			PolygonsToTriangulate.Add( PolygonRef );
+			const FPolygonID PolygonID( PolygonElement.ElementAddress.ElementID );
+			PolygonsToTriangulate.Add( PolygonID );
 		}
 
-		static TArray<FPolygonRef> NewTrianglePolygonRefs;
-		EditableMesh->TriangulatePolygons( PolygonsToTriangulate, /* Out */ NewTrianglePolygonRefs );
+		static TArray<FPolygonID> NewTrianglePolygonIDs;
+		EditableMesh->TriangulatePolygons( PolygonsToTriangulate, /* Out */ NewTrianglePolygonIDs );
 
-		for( const FPolygonRef& NewTrianglePolygonRef : NewTrianglePolygonRefs )
+		for( const FPolygonID NewTrianglePolygonID : NewTrianglePolygonIDs )
 		{
 			// Select the new polygon
 			FMeshElement NewPolygonMeshElement;
@@ -1760,8 +1777,7 @@ bool FMeshEditorMode::TriangulateSelectedPolygons()
 				NewPolygonMeshElement.Component = Component;
 				NewPolygonMeshElement.ElementAddress.SubMeshAddress = EditableMesh->GetSubMeshAddress();
 				NewPolygonMeshElement.ElementAddress.ElementType = EEditableMeshElementType::Polygon;
-				NewPolygonMeshElement.ElementAddress.SectionID = NewTrianglePolygonRef.SectionID;
-				NewPolygonMeshElement.ElementAddress.ElementID = NewTrianglePolygonRef.PolygonID;
+				NewPolygonMeshElement.ElementAddress.ElementID = NewTrianglePolygonID;
 			}
 
 			MeshElementsToSelect.Add( NewPolygonMeshElement );
@@ -1811,8 +1827,8 @@ bool FMeshEditorMode::AssignMaterialToSelectedPolygons( UMaterialInterface* Sele
 			return false;
 		}
 
-		static TArray<FPolygonRef> NewPolygonRefs;
-		NewPolygonRefs.Reset();
+		static TArray<FPolygonID> NewPolygonIDs;
+		NewPolygonIDs.Reset();
 
 		static TArray<FMeshElement> MeshElementsToSelect;
 		MeshElementsToSelect.Reset();
@@ -1842,18 +1858,18 @@ bool FMeshEditorMode::AssignMaterialToSelectedPolygons( UMaterialInterface* Sele
 
 			EditableMesh->StartModification( EMeshModificationType::Final, EMeshTopologyChange::TopologyChange );
 			{
-				static TArray<FPolygonRef> PolygonsToAssign;
+				static TArray<FPolygonID> PolygonsToAssign;
 				PolygonsToAssign.Reset();
 
 				for( const FMeshElement& PolygonElement : MeshAndPolygons.Value )
 				{
-					const FPolygonRef PolygonRef( PolygonElement.ElementAddress.SectionID, FPolygonID( PolygonElement.ElementAddress.ElementID ) );
-					PolygonsToAssign.Add( PolygonRef );
+					const FPolygonID PolygonID( PolygonElement.ElementAddress.ElementID );
+					PolygonsToAssign.Add( PolygonID );
 				}
 
-				EditableMesh->AssignMaterialToPolygons( PolygonsToAssign, SelectedMaterial, NewPolygonRefs );
+				EditableMesh->AssignMaterialToPolygons( PolygonsToAssign, SelectedMaterial, NewPolygonIDs );
 
-				for( const FPolygonRef& NewPolygonRef : NewPolygonRefs )
+				for( const FPolygonID NewPolygonID : NewPolygonIDs )
 				{
 					// Select the new polygon
 					FMeshElement NewPolygonMeshElement;
@@ -1861,8 +1877,7 @@ bool FMeshEditorMode::AssignMaterialToSelectedPolygons( UMaterialInterface* Sele
 						NewPolygonMeshElement.Component = Component;
 						NewPolygonMeshElement.ElementAddress.SubMeshAddress = EditableMesh->GetSubMeshAddress();
 						NewPolygonMeshElement.ElementAddress.ElementType = EEditableMeshElementType::Polygon;
-						NewPolygonMeshElement.ElementAddress.SectionID = NewPolygonRef.SectionID;
-						NewPolygonMeshElement.ElementAddress.ElementID = NewPolygonRef.PolygonID;
+						NewPolygonMeshElement.ElementAddress.ElementID = NewPolygonID;
 					}
 
 					MeshElementsToSelect.Add( NewPolygonMeshElement );
@@ -2020,7 +2035,7 @@ void FMeshEditorMode::AddEdgeToDynamicMesh( const UEditableMesh& EditableMesh, c
 	FVertexID MeshVertexIDs[ 2 ];
 	EditableMesh.GetEdgeVertices( EdgeID, /* Out */ MeshVertexIDs[0], /* Out */ MeshVertexIDs[1] );
 
-	static TArray<FPolygonRef> ConnectedPolygons;
+	static TArray<FPolygonID> ConnectedPolygons;
 	EditableMesh.GetEdgeConnectedPolygons( EdgeID, /* Out */ ConnectedPolygons );
 
 	FVector EdgeVertexPositions[ 2 ];
@@ -2034,11 +2049,11 @@ void FMeshEditorMode::AddEdgeToDynamicMesh( const UEditableMesh& EditableMesh, c
 }
 
 
-void FMeshEditorMode::AddPolygonToDynamicMesh( const UEditableMesh& EditableMesh, const FTransform& CameraToWorld, const bool bIsPerspectiveView, const FMatrix& ComponentToWorldMatrix, const FPolygonRef PolygonRef, const FColor ColorAndOpacity, const float SizeBias, const bool bFillFaces, FDynamicMeshBuilder& VertexAndEdgeMeshBuilder, FDynamicMeshBuilder* PolygonFaceMeshBuilder )
+void FMeshEditorMode::AddPolygonToDynamicMesh( const UEditableMesh& EditableMesh, const FTransform& CameraToWorld, const bool bIsPerspectiveView, const FMatrix& ComponentToWorldMatrix, const FPolygonID PolygonID, const FColor ColorAndOpacity, const FColor HardEdgeColorAndOpacity, const float SizeBias, const bool bFillFaces, FDynamicMeshBuilder& VertexAndEdgeMeshBuilder, FDynamicMeshBuilder* PolygonFaceMeshBuilder )
 {
 	static TArray<FVertexID> MeshVertexIDs;
 	MeshVertexIDs.Reset();
-	EditableMesh.GetPolygonPerimeterVertices( PolygonRef, /* Out */ MeshVertexIDs );
+	EditableMesh.GetPolygonPerimeterVertices( PolygonID, /* Out */ MeshVertexIDs );
 
 	static TArray<FVector> PolygonPerimeterVertexPositions;
 	PolygonPerimeterVertexPositions.SetNumUninitialized( MeshVertexIDs.Num(), false );
@@ -2050,7 +2065,7 @@ void FMeshEditorMode::AddPolygonToDynamicMesh( const UEditableMesh& EditableMesh
 	// @todo mesheditor: Support polygon hole contours!
 
 
-	const FVector PolygonNormal = ComponentToWorldMatrix.TransformVector( EditableMesh.ComputePolygonNormal( PolygonRef ) ).GetSafeNormal();
+	const FVector PolygonNormal = ComponentToWorldMatrix.TransformVector( EditableMesh.ComputePolygonNormal( PolygonID ) ).GetSafeNormal();
 	const float Determinant = ComponentToWorldMatrix.Determinant();
 
 	const FVector CameraForwardVectorUnnormalized = CameraToWorld.TransformVector( FVector::ForwardVector );
@@ -2068,7 +2083,7 @@ void FMeshEditorMode::AddPolygonToDynamicMesh( const UEditableMesh& EditableMesh
 		const float DistanceBias = bIsPerspectiveView ? MeshEd::OverlayPerspectiveDistanceBias->GetFloat() : MeshEd::OverlayOrthographicDistanceBias->GetFloat();
 
 		check( PolygonFaceMeshBuilder );
-		const uint32 PolygonTriangleCount = EditableMesh.GetPolygonTriangulatedTriangleCount( PolygonRef );
+		const uint32 PolygonTriangleCount = EditableMesh.GetPolygonTriangulatedTriangleCount( PolygonID );
 		for( uint32 PolygonTriangleNumber = 0; PolygonTriangleNumber < PolygonTriangleCount; ++PolygonTriangleNumber )
 		{
 			FVector TriangleVertexPositions[ 3 ];
@@ -2077,7 +2092,7 @@ void FMeshEditorMode::AddPolygonToDynamicMesh( const UEditableMesh& EditableMesh
 			{
 				TriangleVertexPositions[ TriangleVertexNumber ] =
 					ComponentToWorldMatrix.TransformPosition( 
-						EditableMesh.GetPolygonTriangulatedTriangleVertexPosition( PolygonRef, PolygonTriangleNumber, TriangleVertexNumber ) );
+						EditableMesh.GetPolygonTriangulatedTriangleVertexPosition( PolygonID, PolygonTriangleNumber, TriangleVertexNumber ) );
 
 				// We're offsetting the geometry from the actual face a bit, to avoid z-fighting for this particular effect
 				const FVector DirectionToCamera = bIsPerspectiveView ? ( CameraToWorld.GetLocation() - TriangleVertexPositions[ TriangleVertexNumber ] ).GetSafeNormal() : -CameraForwardVectorUnnormalized.GetSafeNormal();
@@ -2100,13 +2115,20 @@ void FMeshEditorMode::AddPolygonToDynamicMesh( const UEditableMesh& EditableMesh
 
 	{
 		static TArray<FEdgeID> PerimeterEdges;
-		EditableMesh.GetPolygonPerimeterEdges( PolygonRef, /* Out */ PerimeterEdges );
+		EditableMesh.GetPolygonPerimeterEdges( PolygonID, /* Out */ PerimeterEdges );
 
 		// Draw polygon edges
 		// @todo mesheditor: Assumes no holes, etc.
 		for( const FEdgeID EdgeID : PerimeterEdges )
 		{
-			AddEdgeToDynamicMesh( EditableMesh, CameraToWorld, bIsPerspectiveView, ComponentToWorldMatrix, EdgeID, ColorAndOpacity, SizeBias, VertexAndEdgeMeshBuilder );
+			if( EditableMesh.GetEdgeAttribute( EdgeID, UEditableMeshAttribute::EdgeIsHard(), 0 ).X > 0.0f )
+			{
+				AddEdgeToDynamicMesh( EditableMesh, CameraToWorld, bIsPerspectiveView, ComponentToWorldMatrix, EdgeID, HardEdgeColorAndOpacity, SizeBias, VertexAndEdgeMeshBuilder );
+			}
+			else
+			{
+				AddEdgeToDynamicMesh( EditableMesh, CameraToWorld, bIsPerspectiveView, ComponentToWorldMatrix, EdgeID, ColorAndOpacity, SizeBias, VertexAndEdgeMeshBuilder );
+			}
 		}
 	}
 }
@@ -2178,7 +2200,8 @@ void FMeshEditorMode::DrawMeshElements( const FTransform& CameraToWorld, const b
 								CameraToWorld,
 								bIsPerspectiveView,
 								ComponentToWorldMatrix,
-								FPolygonRef( MeshElement.ElementAddress.SectionID, FPolygonID( MeshElement.ElementAddress.ElementID ) ),
+								FPolygonID( MeshElement.ElementAddress.ElementID ),
+								ElementColor,
 								ElementColor,
 								ElementSizeBias,
 								bFillFaces,
@@ -2241,22 +2264,23 @@ void FMeshEditorMode::Render( const FSceneView* SceneView, FViewport* Viewport,
 			GEngine->AddOnScreenDebugMessage( INDEX_NONE, 0.0f, FColor::White,
 				FString::Printf( TEXT( "Vertices: %i (array size: %i)" ), EditableMesh->GetVertexCount(), EditableMesh->GetVertexArraySize() ), false );
 			GEngine->AddOnScreenDebugMessage( INDEX_NONE, 0.0f, FColor::White,
-				FString::Printf( TEXT( "Rendering Vertices: %i (array size: %i)" ), EditableMesh->GetRenderingVertexCount(), EditableMesh->GetRenderingVertexArraySize() ), false );
+				FString::Printf( TEXT( "Vertex instances: %i (array size: %i)" ), EditableMesh->GetVertexInstanceCount(), EditableMesh->GetVertexInstanceArraySize() ), false );
 			GEngine->AddOnScreenDebugMessage( INDEX_NONE, 0.0f, FColor::White,
 				FString::Printf( TEXT( "Edges: %i (array size: %i)" ), EditableMesh->GetEdgeCount(), EditableMesh->GetEdgeArraySize() ), false );
 			GEngine->AddOnScreenDebugMessage( INDEX_NONE, 0.0f, FColor::White,
-				FString::Printf( TEXT( "Sections: %i (array size: %i)" ), EditableMesh->GetSectionCount(), EditableMesh->GetSectionArraySize() ), false );
-			for( int32 SectionIndex = 0; SectionIndex < EditableMesh->GetSectionArraySize(); ++SectionIndex )
-			{
-				const FSectionID SectionID( SectionIndex );
-				if( EditableMesh->IsValidSection( SectionID ) )
-				{
-					GEngine->AddOnScreenDebugMessage( INDEX_NONE, 0.0f, FColor::White,
-						FString::Printf( TEXT( "   [%i] Polygons: %i (array size: %i)" ), SectionID.GetValue(), EditableMesh->GetPolygonCount( SectionID ), EditableMesh->GetPolygonArraySize( SectionID ) ), false );
-					GEngine->AddOnScreenDebugMessage( INDEX_NONE, 0.0f, FColor::White,
-						FString::Printf( TEXT( "   [%i] Triangles: %i (array size: %i)" ), SectionID.GetValue(), EditableMesh->GetTriangleCount( SectionID ), EditableMesh->GetTriangleArraySize( SectionID ) ), false );
-				}
-			}
+				FString::Printf( TEXT( "Polygons: %i (array size: %i)" ), EditableMesh->GetPolygonCount(), EditableMesh->GetPolygonArraySize() ), false );
+			GEngine->AddOnScreenDebugMessage( INDEX_NONE, 0.0f, FColor::White,
+				FString::Printf( TEXT( "Sections: %i (array size: %i)" ), EditableMesh->GetPolygonGroupCount(), EditableMesh->GetPolygonGroupArraySize() ), false );
+			// @todo mesheditor: triangles array is now an implementation detail in the adapter. Discuss if it's worth preserving access to it somehow.
+			//for( int32 PolygonGroupIndex = 0; PolygonGroupIndex < EditableMesh->GetPolygonGroupArraySize(); ++PolygonGroupIndex )
+			//{
+			//	const FPolygonGroupID PolygonGroupID( PolygonGroupIndex );
+			//	if( EditableMesh->IsValidPolygonGroup( PolygonGroupID ) )
+			//	{
+			//		GEngine->AddOnScreenDebugMessage( INDEX_NONE, 0.0f, FColor::White,
+			//			FString::Printf( TEXT( "   [%i] Triangles: %i (array size: %i)" ), PolygonGroupID.GetValue(), EditableMesh->GetTriangleCount( PolygonGroupID ), EditableMesh->GetTriangleArraySize( PolygonGroupID ) ), false );
+			//	}
+			//}
 		}
 	}
 
@@ -2320,7 +2344,7 @@ void FMeshEditorMode::Render( const FSceneView* SceneView, FViewport* Viewport,
 							// Find all of this polygons edges.  They'll all need to be highlighted.
 							static TArray<FEdgeID> SelectedPolygonPerimeterEdgeIDs;
 							EditableMesh->GetPolygonPerimeterEdges(
-								FPolygonRef( SelectedMeshElement.ElementAddress.SectionID, FPolygonID( SelectedMeshElement.ElementAddress.ElementID ) ),
+								FPolygonID( SelectedMeshElement.ElementAddress.ElementID ),
 								/* Out */ SelectedPolygonPerimeterEdgeIDs );
 
 							HighlightedEdgeIDs.Append( SelectedPolygonPerimeterEdgeIDs );
@@ -2377,62 +2401,71 @@ void FMeshEditorMode::Render( const FSceneView* SceneView, FViewport* Viewport,
 				const float SizeBias = 0.05f;
 				const bool bFillFaces = false;
 				const FColor ColorAndOpacity = FLinearColor( 0.0f, 0.0f, 0.3f ).CopyWithNewOpacity( Opacity ).ToFColor( false );
+				const FColor HardEdgeColorAndOpacity = FLinearColor( 0.0f, 0.3f, 0.3f ).CopyWithNewOpacity( Opacity ).ToFColor( false );
+
+				const uint32 PolygonArraySize = EditableMesh->GetPolygonArraySize();
 
-				for( int32 SectionIndex = 0; SectionIndex < EditableMesh->GetSectionArraySize(); ++SectionIndex )
+				for( uint32 PolygonIndex = 0; PolygonIndex < PolygonArraySize; ++PolygonIndex )
 				{
-					const FSectionID SectionID( SectionIndex );
-					if( EditableMesh->IsValidSection( SectionID ) )
+					const FPolygonID PolygonID( PolygonIndex );
+					if( EditableMesh->IsValidPolygon( PolygonID ) )
 					{
-						const uint32 SectionPolygonArraySize = EditableMesh->GetPolygonArraySize( SectionID );
-
-						for( uint32 PolygonIndex = 0; PolygonIndex < SectionPolygonArraySize; ++PolygonIndex )
+						AddPolygonToDynamicMesh(
+							*EditableMesh,
+							CameraToWorld,
+							bIsPerspectiveView,
+							ComponentToWorldMatrix,
+							PolygonID,
+							ColorAndOpacity,
+							HardEdgeColorAndOpacity,
+							SizeBias,
+							bFillFaces,
+							VertexAndEdgeMeshBuilder,
+							nullptr );
+
+						if( bShowVertexNormals )
 						{
-							const FPolygonID PolygonID( PolygonIndex );
+							// @todo mesheditor: total debug feature for now. Need a way of making this look nice.
+							const int32 PerimeterVertexCount = EditableMesh->GetPolygonPerimeterVertexCount( PolygonID );
+							const float Length = 10.0f; // @todo mesheditor: determine length of debug line from distance from the mesh origin to the camera?
 
-							const FPolygonRef PolygonRef( SectionID, PolygonID );
-							if( EditableMesh->IsValidPolygon( PolygonRef ) )
+							for( int32 PerimeterVertexIndex = 0; PerimeterVertexIndex < PerimeterVertexCount; ++PerimeterVertexIndex )
 							{
-								AddPolygonToDynamicMesh(
-									*EditableMesh,
-									CameraToWorld,
-									bIsPerspectiveView,
-									ComponentToWorldMatrix,
-									PolygonRef,
-									ColorAndOpacity,
-									SizeBias,
-									bFillFaces,
-									VertexAndEdgeMeshBuilder,
-									nullptr );
+								FVector Position( EditableMesh->GetPolygonPerimeterVertexAttribute( PolygonID, PerimeterVertexIndex, UEditableMeshAttribute::VertexPosition(), 0 ) );
+								FVector Normal( EditableMesh->GetPolygonPerimeterVertexAttribute( PolygonID, PerimeterVertexIndex, UEditableMeshAttribute::VertexNormal(), 0 ) );
+								FVector Tangent( EditableMesh->GetPolygonPerimeterVertexAttribute( PolygonID, PerimeterVertexIndex, UEditableMeshAttribute::VertexTangent(), 0 ) );
 
-								if( bShowVertexNormals )
+								FVector VertexNormalRenderPositions[] =
 								{
-									// @todo mesheditor: total debug feature for now. Need a way of making this look nice.
-									const int32 PerimeterVertexCount = EditableMesh->GetPolygonPerimeterVertexCount( PolygonRef );
-									const float Length = 10.0f; // @todo mesheditor: determine length of debug line from distance from the mesh origin to the camera?
+									ComponentToWorldMatrix.TransformPosition( Position ),
+									ComponentToWorldMatrix.TransformPosition( Position + Normal * Length )
+								};
 
-									for( int32 PerimeterVertexIndex = 0; PerimeterVertexIndex < PerimeterVertexCount; ++PerimeterVertexIndex )
-									{
-										FVector Position( EditableMesh->GetPolygonPerimeterVertexAttribute( PolygonRef, PerimeterVertexIndex, UEditableMeshAttribute::VertexPosition(), 0 ) );
-										FVector Normal( EditableMesh->GetPolygonPerimeterVertexAttribute( PolygonRef, PerimeterVertexIndex, UEditableMeshAttribute::VertexNormal(), 0 ) );
+								FVector VertexTangentRenderPositions[] =
+								{
+									ComponentToWorldMatrix.TransformPosition( Position ),
+									ComponentToWorldMatrix.TransformPosition( Position + Tangent * Length * 0.5f )
+								};
 
-										FVector VertexNormalRenderPositions[] =
-										{
-											ComponentToWorldMatrix.TransformPosition( Position ),
-											ComponentToWorldMatrix.TransformPosition( Position + Normal * Length )
-										};
-
-										// We don't need any depth bias when drawing lines for face normals
-										const bool bApplyDepthBias = false;
-										AddThickLineToDynamicMesh(
-											CameraToWorld,
-											bIsPerspectiveView,
-											VertexNormalRenderPositions,
-											FColor::Magenta,
-											SizeBias,
-											bApplyDepthBias,
-											VertexAndEdgeMeshBuilder );
-									}
-								}
+								// We don't need any depth bias when drawing lines for face normals
+								const bool bApplyDepthBias = false;
+								AddThickLineToDynamicMesh(
+									CameraToWorld,
+									bIsPerspectiveView,
+									VertexNormalRenderPositions,
+									FColor::Magenta,
+									SizeBias,
+									bApplyDepthBias,
+									VertexAndEdgeMeshBuilder );
+
+								AddThickLineToDynamicMesh(
+									CameraToWorld,
+									bIsPerspectiveView,
+									VertexTangentRenderPositions,
+									FColor::Yellow,
+									SizeBias,
+									bApplyDepthBias,
+									VertexAndEdgeMeshBuilder );
 							}
 						}
 					}
@@ -3212,30 +3245,21 @@ void FMeshEditorMode::UpdateActiveAction( const bool bIsActionFinishing )
 				else if( DrawnPoints.Num() > 2 )
 				{
 					// If more than two points, create a polygon
-					static TArray<FPolygonRef> NewPolygonRefs;
+					static TArray<FPolygonID> NewPolygonIDs;
 					static TArray<FEdgeID> NewEdgeIDs;
 					static TArray<FPolygonToCreate> PolygonsToCreate;
-					NewPolygonRefs.Reset( 1 );
+					NewPolygonIDs.Reset( 1 );
 					NewEdgeIDs.Reset( 1 );
 					PolygonsToCreate.Reset( 1 );
 
-					// Find first valid section to add the polygon to
-					FSectionID SectionID = FSectionID::Invalid;
-					const int32 SectionArraySize = EditableMesh->GetSectionArraySize();
-					for( int32 SectionIndex = 0; SectionIndex < SectionArraySize; ++SectionIndex )
-					{
-						if( EditableMesh->IsValidSection( FSectionID( SectionIndex ) ) )
-						{
-							SectionID = FSectionID( SectionIndex );
-							break;
-						}
-					}
-					check( SectionID != FSectionID::Invalid );
+					// Find first valid polygon group to add the polygon to
+					FPolygonGroupID PolygonGroupID = EditableMesh->GetFirstValidPolygonGroup();
+					check( PolygonGroupID != FPolygonGroupID::Invalid );
 
 					// Create new polygon
 					PolygonsToCreate.Emplace();
 					FPolygonToCreate& PolygonToCreate = PolygonsToCreate.Last();
-					PolygonToCreate.SectionID = SectionID;
+					PolygonToCreate.PolygonGroupID = PolygonGroupID;
 
 					for( FVertexID NewVertexID : NewVertexIDs )
 					{
@@ -3244,17 +3268,17 @@ void FMeshEditorMode::UpdateActiveAction( const bool bIsActionFinishing )
 						VertexAndAttributes.VertexID = NewVertexID;
 					}
 
-					EditableMesh->CreatePolygons( PolygonsToCreate, NewPolygonRefs, NewEdgeIDs );
+					EditableMesh->CreatePolygons( PolygonsToCreate, NewPolygonIDs, NewEdgeIDs );
 
 					// Check if the polygon normal is pointing towards us. If not, we need to flip the polygon
-					FVector PolygonNormal = EditableMesh->ComputePolygonNormal( NewPolygonRefs[ 0 ] );
+					FVector PolygonNormal = EditableMesh->ComputePolygonNormal( NewPolygonIDs[ 0 ] );
 
 					// @todo mesheditor: Add support for backface checks in orthographic mode
 					if( CachedCameraToWorld.IsSet() && bCachedIsPerspectiveView.IsSet() && !bCachedIsPerspectiveView.GetValue() )
 					{
 						if( FVector::DotProduct( Component->GetComponentTransform().TransformVector( PolygonNormal ), DrawnPoints[ 0 ].Get<1>() - CachedCameraToWorld.GetValue().GetLocation() ) > 0.0f )
 						{
-							EditableMesh->FlipPolygons( NewPolygonRefs );
+							EditableMesh->FlipPolygons( NewPolygonIDs );
 							PolygonNormal = -PolygonNormal;
 						}
 					}
@@ -3265,7 +3289,7 @@ void FMeshEditorMode::UpdateActiveAction( const bool bIsActionFinishing )
 
 					VertexAttributesForPolygon.Emplace();
 					FVertexAttributesForPolygon& VertexAttrs = VertexAttributesForPolygon.Last();
-					VertexAttrs.PolygonRef = NewPolygonRefs[ 0 ];
+					VertexAttrs.PolygonID = NewPolygonIDs[ 0 ];
 
 					for( int32 PolygonVertexIndex = 0; PolygonVertexIndex < NewVertexIDs.Num(); ++PolygonVertexIndex )
 					{
@@ -3415,10 +3439,10 @@ void FMeshEditorMode::UpdateActiveAction( const bool bIsActionFinishing )
 				}
 				else if( ElementToMove.ElementAddress.ElementType == EEditableMeshElementType::Polygon )
 				{
-					const FPolygonRef PolygonRef( ElementToMove.ElementAddress.SectionID, FPolygonID( ElementToMove.ElementAddress.ElementID ) );
+					const FPolygonID PolygonID( ElementToMove.ElementAddress.ElementID );
 
 					static TArray<FVertexID> PolygonPerimeterVertexIDs;
-					EditableMesh->GetPolygonPerimeterVertices( PolygonRef, /* Out */ PolygonPerimeterVertexIDs );
+					EditableMesh->GetPolygonPerimeterVertices( PolygonID, /* Out */ PolygonPerimeterVertexIDs );
 
 					for( const FVertexID PolygonPerimeterVertexID : PolygonPerimeterVertexIDs )
 					{
@@ -3577,37 +3601,29 @@ FEditableMeshElementAddress FMeshEditorMode::QueryElement( const UEditableMesh&
 	static TSet<FEdgeID> FrontFacingEdges;
 	FrontFacingEdges.Reset();
 
-	static TSet<FPolygonRef> FrontFacingPolygons;
+	static TSet<FPolygonID> FrontFacingPolygons;
 	FrontFacingPolygons.Reset();
 
 	// Look for all the front-facing elements
-	for( int32 SectionIndex = 0; SectionIndex < EditableMesh.GetSectionArraySize(); ++SectionIndex )
+	const uint32 PolygonArraySize = EditableMesh.GetPolygonArraySize();
+	for( uint32 PolygonIndex = 0; PolygonIndex < PolygonArraySize; ++PolygonIndex )
 	{
-		const FSectionID SectionID( SectionIndex );
-		if( EditableMesh.IsValidSection( SectionID ) )
+		const FPolygonID PolygonID( PolygonIndex );
+
+		if( EditableMesh.IsValidPolygon( PolygonID ) )
 		{
-			const uint32 SectionPolygonArraySize = EditableMesh.GetPolygonArraySize( SectionID );
-			for( uint32 SectionPolygonIndex = 0; SectionPolygonIndex < SectionPolygonArraySize; ++SectionPolygonIndex )
+			const FVector PolygonNormal = EditableMesh.ComputePolygonNormal( PolygonID );
+			const FVector PolygonCenter = EditableMesh.ComputePolygonCenter( PolygonID );
+			if( !bIsPerspectiveView || FVector::DotProduct( CameraLocation - PolygonCenter, PolygonNormal ) > 0.0f )	// @todo mesheditor: Add support for backface culling in orthographic views
 			{
-				const FPolygonID PolygonID( SectionPolygonIndex );
+				FrontFacingPolygons.Add( PolygonID );
 
-				const FPolygonRef PolygonRef( SectionID, PolygonID );
-				if( EditableMesh.IsValidPolygon( PolygonRef ) )
+				const int32 PolygonVertexCount = EditableMesh.GetPolygonPerimeterVertexCount( PolygonID );
+				for( int32 Index = 0; Index < PolygonVertexCount; ++Index )
 				{
-					const FVector PolygonNormal = EditableMesh.ComputePolygonNormal( PolygonRef );
-					const FVector PolygonCenter = EditableMesh.ComputePolygonCenter( PolygonRef );
-					if( !bIsPerspectiveView || FVector::DotProduct( CameraLocation - PolygonCenter, PolygonNormal ) > 0.0f )	// @todo mesheditor: Add support for backface culling in orthographic views
-					{
-						FrontFacingPolygons.Add( PolygonRef );
-
-						const int32 PolygonVertexCount = EditableMesh.GetPolygonPerimeterVertexCount( PolygonRef );
-						for( int32 Index = 0; Index < PolygonVertexCount; ++Index )
-						{
-							FrontFacingVertices.Add( EditableMesh.GetPolygonPerimeterVertex( PolygonRef, Index ) );
-							bool bOutEdgeWindingIsReversedForPolygons;
-							FrontFacingEdges.Add( EditableMesh.GetPolygonPerimeterEdge( PolygonRef, Index, bOutEdgeWindingIsReversedForPolygons ) );
-						}
-					}
+					FrontFacingVertices.Add( EditableMesh.GetPolygonPerimeterVertex( PolygonID, Index ) );
+					bool bOutEdgeWindingIsReversedForPolygons;
+					FrontFacingEdges.Add( EditableMesh.GetPolygonPerimeterEdge( PolygonID, Index, bOutEdgeWindingIsReversedForPolygons ) );
 				}
 			}
 		}
@@ -3620,20 +3636,20 @@ FEditableMeshElementAddress FMeshEditorMode::QueryElement( const UEditableMesh&
 	FVector CurrentRayEnd = RayEnd;
 
 	// Check polygons first; this is so we always impose a closest hit location at the poly before checking other elements, so anything behind is occluded
-	for( const FPolygonRef PolygonRef : FrontFacingPolygons )
+	for( const FPolygonID PolygonID : FrontFacingPolygons )
 	{
 		static TArray<FVertexID> MeshVertexIDs;
 		MeshVertexIDs.Reset();
-		EditableMesh.GetPolygonPerimeterVertices( PolygonRef, /* Out */ MeshVertexIDs );
+		EditableMesh.GetPolygonPerimeterVertices( PolygonID, /* Out */ MeshVertexIDs );
 
-		const uint32 PolygonTriangleCount = EditableMesh.GetPolygonTriangulatedTriangleCount( PolygonRef );
+		const uint32 PolygonTriangleCount = EditableMesh.GetPolygonTriangulatedTriangleCount( PolygonID );
 		for( uint32 PolygonTriangleNumber = 0; PolygonTriangleNumber < PolygonTriangleCount; ++PolygonTriangleNumber )
 		{
 			FVector TriangleVertexPositions[ 3 ];
 			for( uint32 TriangleVertexNumber = 0; TriangleVertexNumber < 3; ++TriangleVertexNumber )
 			{
 				TriangleVertexPositions[ TriangleVertexNumber ] =
-					EditableMesh.GetPolygonTriangulatedTriangleVertexPosition( PolygonRef, PolygonTriangleNumber, TriangleVertexNumber );
+					EditableMesh.GetPolygonTriangulatedTriangleVertexPosition( PolygonID, PolygonTriangleNumber, TriangleVertexNumber );
 			}
 
 			const bool bAlreadyHitTriangle = ( HitElementAddress.ElementType == EEditableMeshElementType::Polygon );
@@ -3641,8 +3657,7 @@ FEditableMeshElementAddress FMeshEditorMode::QueryElement( const UEditableMesh&
 			if( bHit )
 			{
 				HitElementAddress.ElementType = EEditableMeshElementType::Polygon;
-				HitElementAddress.SectionID = PolygonRef.SectionID;
-				HitElementAddress.ElementID = PolygonRef.PolygonID;
+				HitElementAddress.ElementID = PolygonID;
 			}
 		}
 	}
@@ -3664,7 +3679,6 @@ FEditableMeshElementAddress FMeshEditorMode::QueryElement( const UEditableMesh&
 			if( bHit )
 			{
 				HitElementAddress.ElementType = EEditableMeshElementType::Edge;
-				HitElementAddress.SectionID = FSectionID::Invalid;
 				HitElementAddress.ElementID = EdgeID;
 			}
 		}
@@ -3685,7 +3699,6 @@ FEditableMeshElementAddress FMeshEditorMode::QueryElement( const UEditableMesh&
 			if( bHit )
 			{
 				HitElementAddress.ElementType = EEditableMeshElementType::Vertex;
-				HitElementAddress.SectionID = FSectionID::Invalid;
 				HitElementAddress.ElementID = VertexID;
 			}
 		}
@@ -4423,8 +4436,8 @@ bool FMeshEditorMode::FrustumSelect( const FConvexVolume& InFrustum, FEditorView
 		static TArray<FEdgeID> SelectedEdgeIDs;
 		SelectedEdgeIDs.Empty();
 
-		static TSet<FPolygonRef> SelectedPolygonRefs;
-		SelectedPolygonRefs.Empty();
+		static TSet<FPolygonID> SelectedPolygonIDs;
+		SelectedPolygonIDs.Empty();
 
 		static TSet<FVertexID> SelectedVertexIDs;
 		SelectedVertexIDs.Empty();
@@ -4451,7 +4464,7 @@ bool FMeshEditorMode::FrustumSelect( const FConvexVolume& InFrustum, FEditorView
 					const int32 EdgeConnectedPolygonCount = EditableMesh->GetEdgeConnectedPolygonCount( EdgeID );
 					for( int32 EdgeConnectedPolygonIndex = 0; EdgeConnectedPolygonIndex < EdgeConnectedPolygonCount; ++EdgeConnectedPolygonIndex )
 					{
-						const FPolygonRef EdgeConnectedPolygonID( EditableMesh->GetEdgeConnectedPolygon( EdgeID, EdgeConnectedPolygonIndex ) );
+						const FPolygonID EdgeConnectedPolygonID( EditableMesh->GetEdgeConnectedPolygon( EdgeID, EdgeConnectedPolygonIndex ) );
 
 						// Determine whether polygon is back facing or not using dot product of its normal with the direction vector from the eye position to somewhere on the plane
 						// (one of the vertex positions is sufficient for this)
@@ -4464,7 +4477,7 @@ bool FMeshEditorMode::FrustumSelect( const FConvexVolume& InFrustum, FEditorView
 						if( !bOnlySelectVisibleElements || !bIsBackFacing )
 						{
 							// Add the polygon if it is front facing, or if we don't care about only selecting visible elements
-							SelectedPolygonRefs.Add( EdgeConnectedPolygonID );
+							SelectedPolygonIDs.Add( EdgeConnectedPolygonID );
 						}
 					}
 
@@ -4532,10 +4545,10 @@ bool FMeshEditorMode::FrustumSelect( const FConvexVolume& InFrustum, FEditorView
 
 		if( MeshElementSelectionMode == EEditableMeshElementType::Polygon || MeshElementSelectionMode == EEditableMeshElementType::Any )
 		{
-			MarqueeSelectPolygons.Reserve( MarqueeSelectPolygons.Num() + SelectedPolygonRefs.Num() );
-			for( FPolygonRef SelectedPolygonRef : SelectedPolygonRefs )
+			MarqueeSelectPolygons.Reserve( MarqueeSelectPolygons.Num() + SelectedPolygonIDs.Num() );
+			for( const FPolygonID SelectedPolygonID : SelectedPolygonIDs )
 			{
-				MarqueeSelectPolygons.Emplace( Component, EditableMesh->GetSubMeshAddress(), SelectedPolygonRef );
+				MarqueeSelectPolygons.Emplace( Component, EditableMesh->GetSubMeshAddress(), SelectedPolygonID );
 			}
 		}
 	}
@@ -4685,16 +4698,16 @@ void FMeshEditorMode::RefreshTransformables( const bool bNewObjectsSelected )
 
 						case EEditableMeshElementType::Polygon:
 						{
-							const FPolygonRef PolygonRef( MeshElement.ElementAddress.SectionID, FPolygonID( MeshElement.ElementAddress.ElementID ) );
+							const FPolygonID PolygonID( MeshElement.ElementAddress.ElementID );
 
-							const FVector ComponentSpacePolygonCenter = EditableMesh->ComputePolygonCenter( PolygonRef );
+							const FVector ComponentSpacePolygonCenter = EditableMesh->ComputePolygonCenter( PolygonID );
 							ElementTransform.SetLocation( ComponentToWorldMatrix.TransformPosition( ComponentSpacePolygonCenter ) );
 
 							const FVector WindingVector =
-								( ComponentToWorldMatrix.TransformPosition( EditableMesh->GetVertexAttribute( EditableMesh->GetPolygonPerimeterVertex( PolygonRef, 1 ), UEditableMeshAttribute::VertexPosition(), 0 ) ) -
-									ComponentToWorldMatrix.TransformPosition( EditableMesh->GetVertexAttribute( EditableMesh->GetPolygonPerimeterVertex( PolygonRef, 0 ), UEditableMeshAttribute::VertexPosition(), 0 ) ) ).GetSafeNormal();
+								( ComponentToWorldMatrix.TransformPosition( EditableMesh->GetVertexAttribute( EditableMesh->GetPolygonPerimeterVertex( PolygonID, 1 ), UEditableMeshAttribute::VertexPosition(), 0 ) ) -
+									ComponentToWorldMatrix.TransformPosition( EditableMesh->GetVertexAttribute( EditableMesh->GetPolygonPerimeterVertex( PolygonID, 0 ), UEditableMeshAttribute::VertexPosition(), 0 ) ) ).GetSafeNormal();
 
-							const FVector PolygonNormal = ComponentToWorld.TransformVectorNoScale( EditableMesh->ComputePolygonNormal( PolygonRef ) ).GetSafeNormal();
+							const FVector PolygonNormal = ComponentToWorld.TransformVectorNoScale( EditableMesh->ComputePolygonNormal( PolygonID ) ).GetSafeNormal();
 
 							const FVector PolygonBinormal = FVector::CrossProduct( PolygonNormal, WindingVector ).GetSafeNormal();
 							const FVector PolygonTangent = FVector::CrossProduct( PolygonBinormal, PolygonNormal );
diff --git a/Engine/Plugins/Editor/MeshEditor/Source/MeshEditor/MeshEditorMode.h b/Engine/Plugins/Editor/MeshEditor/Source/MeshEditor/MeshEditorMode.h
--- a/Engine/Plugins/Editor/MeshEditor/Source/MeshEditor/MeshEditorMode.h
+++ b/Engine/Plugins/Editor/MeshEditor/Source/MeshEditor/MeshEditorMode.h
@@ -13,6 +13,9 @@
 #include "UnrealEdMisc.h"
 #include "IMeshEditorModeUIContract.h"
 #include "IMeshEditorModeEditingContract.h"
+#include "Framework/Commands/UIAction.h"
+#include "Framework/Commands/UICommandInfo.h"
+#include "Framework/Commands/UICommandList.h"
 #include "MeshEditorMode.generated.h"
 
 
@@ -124,7 +127,7 @@ public:
 					break;
 
 				case EEditableMeshElementType::Polygon:
-					bIsValid = EditableMesh->IsValidPolygon( FPolygonRef( MeshElement.ElementAddress.SectionID, FPolygonID( MeshElement.ElementAddress.ElementID ) ) );
+					bIsValid = EditableMesh->IsValidPolygon( FPolygonID( MeshElement.ElementAddress.ElementID ) );
 					break;
 			}
 		}
@@ -219,7 +222,7 @@ protected:
 	void AddEdgeToDynamicMesh( const UEditableMesh& EditableMesh, const FTransform& CameraToWorld, const bool bIsPerspectiveView, const FMatrix& ComponentToWorldMatrix, const FEdgeID EdgeID, const FColor ColorAndOpacity, const float SizeBias, class FDynamicMeshBuilder& MeshBuilder );
 
 	/** Fills the specified dynamic mesh builders with primitives to render a polygon and it's edges */
-	void AddPolygonToDynamicMesh( const UEditableMesh& EditableMesh, const FTransform& CameraToWorld, const bool bIsPerspectiveView, const FMatrix& ComponentToWorldMatrix, const FPolygonRef PolygonRef, const FColor ColorAndOpacity, const float SizeBias, const bool bFillFaces, class FDynamicMeshBuilder& VertexAndEdgeMeshBuilder, class FDynamicMeshBuilder* PolygonFaceMeshBuilder );
+	void AddPolygonToDynamicMesh( const UEditableMesh& EditableMesh, const FTransform& CameraToWorld, const bool bIsPerspectiveView, const FMatrix& ComponentToWorldMatrix, const FPolygonID PolygonID, const FColor ColorAndOpacity, const FColor HardEdgeColorAndOpacity, const float SizeBias, const bool bFillFaces, class FDynamicMeshBuilder& VertexAndEdgeMeshBuilder, class FDynamicMeshBuilder* PolygonFaceMeshBuilder );
 
 	/** Renders the spceified mesh element */
 	void DrawMeshElements( const FTransform& CameraToWorld, const bool bIsPerspectiveView, FViewport* Viewport, FPrimitiveDrawInterface* PDI, const TArrayView<FMeshElement>& MeshElements, const FColor Color, const bool bFillFaces, const float HoverAnimation, const TArray<FColor>* OptionalPerElementColors = nullptr, const TArray<float>* OptionalPerElementSizeBiases = nullptr );
diff --git a/Engine/Plugins/Editor/MeshEditor/Source/MeshEditor/MeshElementViewportTransformable.cpp b/Engine/Plugins/Editor/MeshEditor/Source/MeshEditor/MeshElementViewportTransformable.cpp
--- a/Engine/Plugins/Editor/MeshEditor/Source/MeshEditor/MeshElementViewportTransformable.cpp
+++ b/Engine/Plugins/Editor/MeshEditor/Source/MeshEditor/MeshElementViewportTransformable.cpp
@@ -60,10 +60,10 @@ FBox FMeshElementViewportTransformable::BuildBoundingBox( const FTransform& Boun
 
 					case EEditableMeshElementType::Polygon:
 					{
-						const FPolygonRef PolygonRef( MeshElement.ElementAddress.SectionID, FPolygonID( MeshElement.ElementAddress.ElementID ) );
+						const FPolygonID PolygonID( MeshElement.ElementAddress.ElementID );
 
 						static TArray<FVertexID> PerimeterVertexIDs;
-						EditableMesh->GetPolygonPerimeterVertices( PolygonRef, /* Out */ PerimeterVertexIDs );
+						EditableMesh->GetPolygonPerimeterVertices( PolygonID, /* Out */ PerimeterVertexIDs );
 
 						for( const FVertexID VertexID : PerimeterVertexIDs )
 						{
diff --git a/Engine/Plugins/Editor/MeshEditor/Source/MeshEditor/Public/MeshElement.h b/Engine/Plugins/Editor/MeshEditor/Source/MeshEditor/Public/MeshElement.h
--- a/Engine/Plugins/Editor/MeshEditor/Source/MeshEditor/Public/MeshElement.h
+++ b/Engine/Plugins/Editor/MeshEditor/Source/MeshEditor/Public/MeshElement.h
@@ -16,9 +16,6 @@ struct FEditableMeshElementAddress
 	/** The type of element */
 	EEditableMeshElementType ElementType;
 
-	/** For polygons, the section ID the polygon belongs to, otherwise FEditableMeshSectionID::Invalid */
-	FSectionID SectionID;
-
 	/** The ID of the element within the mesh */
 	FElementID ElementID;
 
@@ -27,7 +24,6 @@ struct FEditableMeshElementAddress
 	FEditableMeshElementAddress()
 		: SubMeshAddress(),
 		  ElementType( EEditableMeshElementType::Invalid ),
-		  SectionID( FSectionID::Invalid ),
 		  ElementID( FElementID::Invalid )
 	{
 	}
@@ -35,7 +31,6 @@ struct FEditableMeshElementAddress
 	FEditableMeshElementAddress( const FEditableMeshSubMeshAddress& InSubMeshAddress, FVertexID InVertexID )
 		: SubMeshAddress( InSubMeshAddress ),
 		  ElementType( EEditableMeshElementType::Vertex ),
-		  SectionID( FSectionID::Invalid ),
 		  ElementID( InVertexID )
 	{
 	}
@@ -43,16 +38,14 @@ struct FEditableMeshElementAddress
 	FEditableMeshElementAddress( const FEditableMeshSubMeshAddress& InSubMeshAddress, FEdgeID InEdgeID )
 		: SubMeshAddress( InSubMeshAddress ),
 		  ElementType( EEditableMeshElementType::Edge ),
-		  SectionID( FSectionID::Invalid ),
 		  ElementID( InEdgeID )
 	{
 	}
 
-	FEditableMeshElementAddress( const FEditableMeshSubMeshAddress& InSubMeshAddress, FPolygonRef InPolygonRef )
+	FEditableMeshElementAddress( const FEditableMeshSubMeshAddress& InSubMeshAddress, FPolygonID InPolygonID )
 		: SubMeshAddress( InSubMeshAddress ),
 		  ElementType( EEditableMeshElementType::Polygon ),
-		  SectionID( InPolygonRef.SectionID ),
-		  ElementID( InPolygonRef.PolygonID )
+		  ElementID( InPolygonID )
 	{
 	}
 
@@ -62,7 +55,6 @@ struct FEditableMeshElementAddress
 		return
 			SubMeshAddress == Other.SubMeshAddress &&
 			ElementType == Other.ElementType &&
-			SectionID == Other.SectionID &&
 			ElementID == Other.ElementID;
 	}
 
@@ -93,10 +85,9 @@ struct FEditableMeshElementAddress
 		}
 
 		return FString::Printf(
-			TEXT( "%s, ElementType:%s, SectionID:%s, ElementID:%s" ),
+			TEXT( "%s, ElementType:%s, ElementID:%s" ),
 			*SubMeshAddress.ToString(),
 			*ElementTypeString,
-			*SectionID.ToString(),
 			*ElementID.ToString() );
 	}
 };
@@ -143,9 +134,9 @@ struct FMeshElement
 	{
 	}
 
-	FMeshElement( UPrimitiveComponent* InComponent, const FEditableMeshSubMeshAddress& InSubMeshAddress, FPolygonRef InPolygonRef, double InLastHoverTime = 0.0, double InLastSelectTime = 0.0 )
+	FMeshElement( UPrimitiveComponent* InComponent, const FEditableMeshSubMeshAddress& InSubMeshAddress, FPolygonID InPolygonID, double InLastHoverTime = 0.0, double InLastSelectTime = 0.0 )
 		: Component( InComponent ),
-		  ElementAddress( InSubMeshAddress, InPolygonRef ),
+		  ElementAddress( InSubMeshAddress, InPolygonID ),
 		  LastHoverTime( InLastHoverTime ),
 		  LastSelectTime( InLastSelectTime )
 	{
diff --git a/Engine/Plugins/Editor/MeshEditor/Source/PolygonModeling/BevelOrInsetPolygon.cpp b/Engine/Plugins/Editor/MeshEditor/Source/PolygonModeling/BevelOrInsetPolygon.cpp
--- a/Engine/Plugins/Editor/MeshEditor/Source/PolygonModeling/BevelOrInsetPolygon.cpp
+++ b/Engine/Plugins/Editor/MeshEditor/Source/PolygonModeling/BevelOrInsetPolygon.cpp
@@ -30,7 +30,7 @@ namespace BevelOrInsetPolygonHelpers
 	static void FindInsetAmount(
 		IMeshEditorModeEditingContract& MeshEditorMode,
 		UViewportInteractor* ViewportInteractor,
-		const FPolygonRef& PolygonRef,
+		const FPolygonID PolygonID,
 		UPrimitiveComponent& Component,
 		const UEditableMesh* EditableMesh,
 		float& OutInsetFixedDistance,
@@ -50,8 +50,8 @@ namespace BevelOrInsetPolygonHelpers
 		{
 			const FMatrix ComponentToWorldMatrix = Component.GetRenderMatrix();
 
-			const FPlane PolygonPlane = EditableMesh->ComputePolygonPlane( PolygonRef );
-			const FVector PolygonCenter = EditableMesh->ComputePolygonCenter( PolygonRef );
+			const FPlane PolygonPlane = EditableMesh->ComputePolygonPlane( PolygonID );
+			const FVector PolygonCenter = EditableMesh->ComputePolygonCenter( PolygonID );
 
 			const FVector ComponentSpaceRayStart = ComponentToWorldMatrix.InverseTransformPosition( LaserStart );
 			const FVector ComponentSpaceRayEnd = ComponentToWorldMatrix.InverseTransformPosition( LaserEnd );
@@ -67,7 +67,7 @@ namespace BevelOrInsetPolygonHelpers
 				FVector BestEdgeVertex0Position, BestEdgeVertex1Position;
 				{
 					static TArray<FVertexID> PerimeterVertexIDs;
-					EditableMesh->GetPolygonPerimeterVertices( PolygonRef, /* Out */ PerimeterVertexIDs );
+					EditableMesh->GetPolygonPerimeterVertices( PolygonID, /* Out */ PerimeterVertexIDs );
 
 					for( int32 VertexNumber = 0; VertexNumber < PerimeterVertexIDs.Num(); ++VertexNumber )
 					{
@@ -149,7 +149,7 @@ namespace BevelOrInsetPolygonHelpers
 				const UEditableMesh* InsetUsingEditableMesh = MeshEditorMode.FindEditableMesh( *InsetUsingComponent, InsetUsingPolygonElement.ElementAddress.SubMeshAddress );
 				if( ensure( InsetUsingEditableMesh != nullptr ) )
 				{
-					const FPolygonRef InsetUsingPolygonRef( InsetUsingPolygonElement.ElementAddress.SectionID, FPolygonID( InsetUsingPolygonElement.ElementAddress.ElementID ) );
+					const FPolygonID InsetUsingPolygonID( InsetUsingPolygonElement.ElementAddress.ElementID );
 
 					// Figure out how far to inset the polygon
 					float InsetFixedDistance = 0.0f;
@@ -157,7 +157,7 @@ namespace BevelOrInsetPolygonHelpers
 					FindInsetAmount(
 						MeshEditorMode,
 						ViewportInteractor,
-						InsetUsingPolygonRef,
+						InsetUsingPolygonID,
 						*InsetUsingComponent,
 						InsetUsingEditableMesh,
 						/* Out */ InsetFixedDistance,
@@ -175,30 +175,30 @@ namespace BevelOrInsetPolygonHelpers
 							check( Component != nullptr );
 
 
-							static TArray<FPolygonRef> PolygonRefsToInset;
-							PolygonRefsToInset.Reset();
+							static TArray<FPolygonID> PolygonIDsToInset;
+							PolygonIDsToInset.Reset();
 							for( const FMeshElement& PolygonToInset : PolygonsToInset )
 							{
-								FPolygonRef PolygonRef( PolygonToInset.ElementAddress.SectionID, FPolygonID( PolygonToInset.ElementAddress.ElementID ) );
-								PolygonRefsToInset.Add( PolygonRef );
+								const FPolygonID PolygonID( PolygonToInset.ElementAddress.ElementID );
+								PolygonIDsToInset.Add( PolygonID );
 							}
 
 							{
 								verify( !EditableMesh->AnyChangesToUndo() );
 
 								// Inset time!!
-								static TArray<FPolygonRef> NewCenterInsetPolygons;
-								static TArray<FPolygonRef> NewSideInsetPolygons;
+								static TArray<FPolygonID> NewCenterInsetPolygons;
+								static TArray<FPolygonID> NewSideInsetPolygons;
 								if( BevelOrInset == EBevelOrInset::Inset )
 								{
 									const EInsetPolygonsMode InsetPolygonsMode = EInsetPolygonsMode::All;	// @todo mesheditor inset: Make configurable?
 																											// @todo mesheditor inset: Add options for Fixed distance (instead of Percentage distance, like now.)
 
-									EditableMesh->InsetPolygons( PolygonRefsToInset, InsetFixedDistance, InsetProgressTowardCenter, InsetPolygonsMode, /* Out */ NewCenterInsetPolygons, /* Out */ NewSideInsetPolygons );
+									EditableMesh->InsetPolygons( PolygonIDsToInset, InsetFixedDistance, InsetProgressTowardCenter, InsetPolygonsMode, /* Out */ NewCenterInsetPolygons, /* Out */ NewSideInsetPolygons );
 								}
 								else if( ensure( BevelOrInset == EBevelOrInset::Bevel ) )
 								{
-									EditableMesh->BevelPolygons( PolygonRefsToInset, InsetFixedDistance, InsetProgressTowardCenter, /* Out */ NewCenterInsetPolygons, /* Out */ NewSideInsetPolygons );
+									EditableMesh->BevelPolygons( PolygonIDsToInset, InsetFixedDistance, InsetProgressTowardCenter, /* Out */ NewCenterInsetPolygons, /* Out */ NewSideInsetPolygons );
 								}
 
 								// Make sure the new polygons are selected.  The old polygon was deleted and will become deselected automatically.
@@ -206,7 +206,7 @@ namespace BevelOrInsetPolygonHelpers
 								{
 									for( int32 NewPolygonNumber = 0; NewPolygonNumber < NewCenterInsetPolygons.Num(); ++NewPolygonNumber )
 									{
-										const FPolygonRef& NewInsetCenterPolygon = NewCenterInsetPolygons[ NewPolygonNumber ];
+										const FPolygonID NewInsetCenterPolygon = NewCenterInsetPolygons[ NewPolygonNumber ];
 
 										const FMeshElement& MeshElement = PolygonsToInset[ NewPolygonNumber ];
 
@@ -214,8 +214,7 @@ namespace BevelOrInsetPolygonHelpers
 										{
 											PolygonMeshElement.Component = MeshElement.Component;
 											PolygonMeshElement.ElementAddress = MeshElement.ElementAddress;
-											PolygonMeshElement.ElementAddress.SectionID = NewInsetCenterPolygon.SectionID;
-											PolygonMeshElement.ElementAddress.ElementID = NewInsetCenterPolygon.PolygonID;
+											PolygonMeshElement.ElementAddress.ElementID = NewInsetCenterPolygon;
 										}
 
 										// Queue selection of this new element.  We don't want it to be part of the current action.
@@ -226,7 +225,7 @@ namespace BevelOrInsetPolygonHelpers
 								{
 									for( int32 NewPolygonNumber = 0; NewPolygonNumber < NewSideInsetPolygons.Num(); ++NewPolygonNumber )
 									{
-										const FPolygonRef& NewInsetSidePolygon = NewSideInsetPolygons[ NewPolygonNumber ];
+										const FPolygonID NewInsetSidePolygon = NewSideInsetPolygons[ NewPolygonNumber ];
 
 										const FMeshElement& MeshElement = PolygonsToInset[ NewPolygonNumber ];
 
@@ -234,8 +233,7 @@ namespace BevelOrInsetPolygonHelpers
 										{
 											PolygonMeshElement.Component = MeshElement.Component;
 											PolygonMeshElement.ElementAddress = MeshElement.ElementAddress;
-											PolygonMeshElement.ElementAddress.SectionID = NewInsetSidePolygon.SectionID;
-											PolygonMeshElement.ElementAddress.ElementID = NewInsetSidePolygon.PolygonID;
+											PolygonMeshElement.ElementAddress.ElementID = NewInsetSidePolygon;
 										}
 
 										// Queue selection of this new element.  We don't want it to be part of the current action.
diff --git a/Engine/Plugins/Editor/MeshEditor/Source/PolygonModeling/ExtrudePolygon.cpp b/Engine/Plugins/Editor/MeshEditor/Source/PolygonModeling/ExtrudePolygon.cpp
--- a/Engine/Plugins/Editor/MeshEditor/Source/PolygonModeling/ExtrudePolygon.cpp
+++ b/Engine/Plugins/Editor/MeshEditor/Source/PolygonModeling/ExtrudePolygon.cpp
@@ -97,8 +97,8 @@ bool UExtrudePolygonCommand::TryStartingToDrag( IMeshEditorModeEditingContract&
 						check( Component != nullptr );
 						const FMatrix ComponentToWorldMatrix = Component->GetRenderMatrix();
 
-						const FPolygonRef PolygonRef( PolygonElement.ElementAddress.SectionID, FPolygonID( PolygonElement.ElementAddress.ElementID ) );
-						const FVector ComponentSpacePolygonNormal = EditableMesh->ComputePolygonNormal( PolygonRef );
+						const FPolygonID PolygonID( PolygonElement.ElementAddress.ElementID );
+						const FVector ComponentSpacePolygonNormal = EditableMesh->ComputePolygonNormal( PolygonID );
 
 						const FVector WorldSpacePolygonNormal = ComponentToWorldMatrix.TransformVector( ComponentSpacePolygonNormal ).GetSafeNormal();
 
@@ -133,12 +133,12 @@ void UExtrudePolygonCommand::ApplyDuringDrag( IMeshEditorModeEditingContract& Me
 			UEditableMesh* EditableMesh = MeshAndPolygons.Key;
 			const TArray<FMeshElement>& PolygonsToExtrude = MeshAndPolygons.Value;
 
-			static TArray<FPolygonRef> PolygonRefsToExtrude;
-			PolygonRefsToExtrude.Reset();
+			static TArray<FPolygonID> PolygonIDsToExtrude;
+			PolygonIDsToExtrude.Reset();
 			for( const FMeshElement& PolygonToExtrude : PolygonsToExtrude )
 			{
-				FPolygonRef PolygonRef( PolygonToExtrude.ElementAddress.SectionID, FPolygonID( PolygonToExtrude.ElementAddress.ElementID ) );
-				PolygonRefsToExtrude.Add( PolygonRef );
+				const FPolygonID PolygonID( PolygonToExtrude.ElementAddress.ElementID );
+				PolygonIDsToExtrude.Add( PolygonID );
 			}
 
 
@@ -167,16 +167,16 @@ void UExtrudePolygonCommand::ApplyDuringDrag( IMeshEditorModeEditingContract& Me
 				const float ComponentSpaceExtrudeDistance = ExtrudeDistance / Component->ComponentToWorld.GetScale3D().X;
 
 				// Create a copy of the polygon with new extruded polygons for each edge
-				static TArray<FPolygonRef> NewExtrudedFrontPolygons;
+				static TArray<FPolygonID> NewExtrudedFrontPolygons;
 				NewExtrudedFrontPolygons.Reset();
 				const bool bKeepNeighborsTogether = true;	// @todo mesheditor: Make configurable in UI
-				EditableMesh->ExtrudePolygons( PolygonRefsToExtrude, ComponentSpaceExtrudeDistance, bKeepNeighborsTogether, /* Out */ NewExtrudedFrontPolygons );
+				EditableMesh->ExtrudePolygons( PolygonIDsToExtrude, ComponentSpaceExtrudeDistance, bKeepNeighborsTogether, /* Out */ NewExtrudedFrontPolygons );
 
 
 				// Make sure the new polygons are selected.  The old polygon was deleted and will become deselected automatically.
 				for( int32 NewPolygonNumber = 0; NewPolygonNumber < NewExtrudedFrontPolygons.Num(); ++NewPolygonNumber )
 				{
-					const FPolygonRef& NewExtrudedFrontPolygon = NewExtrudedFrontPolygons[ NewPolygonNumber ];
+					const FPolygonID NewExtrudedFrontPolygon = NewExtrudedFrontPolygons[ NewPolygonNumber ];
 
 					const FMeshElement& MeshElement = PolygonsToExtrude[ NewPolygonNumber ];
 
@@ -184,8 +184,7 @@ void UExtrudePolygonCommand::ApplyDuringDrag( IMeshEditorModeEditingContract& Me
 					{
 						NewExtrudedPolygonMeshElement.Component = MeshElement.Component;
 						NewExtrudedPolygonMeshElement.ElementAddress = MeshElement.ElementAddress;
-						NewExtrudedPolygonMeshElement.ElementAddress.SectionID = NewExtrudedFrontPolygon.SectionID;
-						NewExtrudedPolygonMeshElement.ElementAddress.ElementID = NewExtrudedFrontPolygon.PolygonID;
+						NewExtrudedPolygonMeshElement.ElementAddress.ElementID = NewExtrudedFrontPolygon;
 					}
 
 					// Queue selection of this new element.  We don't want it to be part of the current action.
diff --git a/Engine/Plugins/Editor/MeshEditor/Source/PolygonModeling/QuadrangulateMesh.cpp b/Engine/Plugins/Editor/MeshEditor/Source/PolygonModeling/QuadrangulateMesh.cpp
--- a/Engine/Plugins/Editor/MeshEditor/Source/PolygonModeling/QuadrangulateMesh.cpp
+++ b/Engine/Plugins/Editor/MeshEditor/Source/PolygonModeling/QuadrangulateMesh.cpp
@@ -39,9 +39,11 @@ void UQuadrangulateMeshCommand::Execute( IMeshEditorModeEditingContract& MeshEdi
 
 	for( UEditableMesh* EditableMesh : SelectedMeshes )
 	{
-		static TArray<FPolygonRef> NewPolygonRefs;
 		EditableMesh->StartModification( EMeshModificationType::Final, EMeshTopologyChange::TopologyChange );
-		EditableMesh->QuadrangulateMesh( NewPolygonRefs );
+		{
+			static TArray<FPolygonID> NewPolygonIDs;
+			EditableMesh->QuadrangulateMesh( NewPolygonIDs );
+		}
 		EditableMesh->EndModification();
 
 		MeshEditorMode.TrackUndo( EditableMesh, EditableMesh->MakeUndo() );
diff --git a/Engine/Plugins/Editor/MeshEditor/Source/PolygonModeling/RemoveEdge.cpp b/Engine/Plugins/Editor/MeshEditor/Source/PolygonModeling/RemoveEdge.cpp
--- a/Engine/Plugins/Editor/MeshEditor/Source/PolygonModeling/RemoveEdge.cpp
+++ b/Engine/Plugins/Editor/MeshEditor/Source/PolygonModeling/RemoveEdge.cpp
@@ -6,7 +6,8 @@
 #include "UICommandInfo.h"
 #include "EditableMesh.h"
 #include "MeshElement.h"
-
+#include "Framework/Multibox/MultiBoxBuilder.h"
+#include "ScopedTransaction.h"
 
 #define LOCTEXT_NAMESPACE "MeshEditorMode"
 
@@ -62,9 +63,9 @@ void URemoveEdgeCommand::Execute( IMeshEditorModeEditingContract& MeshEditorMode
 			const FEdgeID EdgeID( EdgeElementToRemove.ElementAddress.ElementID );
 			{
 				bool bWasEdgeRemoved = false;
-				FPolygonRef NewPolygonRef;
+				FPolygonID NewPolygonID;
 
-				EditableMesh->TryToRemovePolygonEdge( EdgeID, /* Out */ bWasEdgeRemoved, /* Out */ NewPolygonRef );
+				EditableMesh->TryToRemovePolygonEdge( EdgeID, /* Out */ bWasEdgeRemoved, /* Out */ NewPolygonID );
 
 				if( bWasEdgeRemoved )
 				{
@@ -74,8 +75,7 @@ void URemoveEdgeCommand::Execute( IMeshEditorModeEditingContract& MeshEditorMode
 						NewPolygonMeshElement.Component = EdgeElementToRemove.Component;
 						NewPolygonMeshElement.ElementAddress = EdgeElementToRemove.ElementAddress;
 						NewPolygonMeshElement.ElementAddress.ElementType = EEditableMeshElementType::Polygon;
-						NewPolygonMeshElement.ElementAddress.SectionID = NewPolygonRef.SectionID;
-						NewPolygonMeshElement.ElementAddress.ElementID = NewPolygonRef.PolygonID;
+						NewPolygonMeshElement.ElementAddress.ElementID = NewPolygonID;
 					}
 
 					MeshElementsToSelect.Add( NewPolygonMeshElement );
diff --git a/Engine/Plugins/Editor/MeshEditor/Source/PolygonModeling/SplitPolygon.cpp b/Engine/Plugins/Editor/MeshEditor/Source/PolygonModeling/SplitPolygon.cpp
--- a/Engine/Plugins/Editor/MeshEditor/Source/PolygonModeling/SplitPolygon.cpp
+++ b/Engine/Plugins/Editor/MeshEditor/Source/PolygonModeling/SplitPolygon.cpp
@@ -114,18 +114,18 @@ void USplitPolygonCommand::ApplyDuringDrag( IMeshEditorModeEditingContract& Mesh
 		// We'll always be trying to split one of the polygons that share the starting vertex.  But we need to figure
 		// out which one of those polygons the user is hovering over (either the polygon itself, or one of it's edges
 		// or vertices.)
-		static TArray<FPolygonRef> CandidatePolygonRefs;
-		CandidatePolygonRefs.Reset();
+		static TArray<FPolygonID> CandidatePolygonIDs;
+		CandidatePolygonIDs.Reset();
 		if( StartingEdgeID != FEdgeID::Invalid )
 		{
-			EditableMesh->GetEdgeConnectedPolygons( StartingEdgeID, /* Out */ CandidatePolygonRefs );
+			EditableMesh->GetEdgeConnectedPolygons( StartingEdgeID, /* Out */ CandidatePolygonIDs );
 		}
 		else if( ensure( StartingVertexID != FVertexID::Invalid ) )
 		{
-			EditableMesh->GetVertexConnectedPolygons( StartingVertexID, /* Out */ CandidatePolygonRefs );
+			EditableMesh->GetVertexConnectedPolygons( StartingVertexID, /* Out */ CandidatePolygonIDs );
 		}
 
-		FPolygonRef PolygonToSplit = FPolygonRef::Invalid;
+		FPolygonID PolygonToSplit = FPolygonID::Invalid;
 		FVertexID ToVertexID = FVertexID::Invalid;
 
 
@@ -133,7 +133,7 @@ void USplitPolygonCommand::ApplyDuringDrag( IMeshEditorModeEditingContract& Mesh
 			FVector LaserPointerStart, LaserPointerEnd;
 			if( ViewportInteractor->GetLaserPointer( /* Out */ LaserPointerStart, /* Out */ LaserPointerEnd ) )
 			{
-				for( const FPolygonRef CandidatePolygonRef : CandidatePolygonRefs )
+				for( const FPolygonID CandidatePolygonID : CandidatePolygonIDs )
 				{
 					const FTransform ComponentToWorld = Component->GetComponentToWorld();
 
@@ -154,14 +154,14 @@ void USplitPolygonCommand::ApplyDuringDrag( IMeshEditorModeEditingContract& Mesh
 					}
 
 
-					const FPlane PolygonPlane = EditableMesh->ComputePolygonPlane( CandidatePolygonRef ).TransformBy( ComponentToWorld.ToMatrixWithScale() );
+					const FPlane PolygonPlane = EditableMesh->ComputePolygonPlane( CandidatePolygonID ).TransformBy( ComponentToWorld.ToMatrixWithScale() );
 					const FVector LaserImpactOnPolygonPlane = FMath::LinePlaneIntersection( LaserPointerStart, LaserPointerEnd, PolygonPlane );
 
 					// @todo mesheditor splitpolygon: Ideally this would be more "fuzzy", and allow the interactor to extend beyond the range of the polygon.  But it will make figuring out which polygon to split more tricky.
 					// @todo mesheditor urgent: Can crash with "Colinear points in FMath::ComputeBaryCentric2D()"  Needs repro.
-					static TArray<int32> PerimeterVertexIndices;
+					FMeshTriangle Triangle;
 					FVector TriangleVertexWeights;
-					if( EditableMesh->ComputeBarycentricWeightForPointOnPolygon( CandidatePolygonRef, ComponentToWorld.InverseTransformPosition( LaserImpactOnPolygonPlane ), /* Out */ PerimeterVertexIndices, /* Out */ TriangleVertexWeights ) )
+					if( EditableMesh->ComputeBarycentricWeightForPointOnPolygon( CandidatePolygonID, ComponentToWorld.InverseTransformPosition( LaserImpactOnPolygonPlane ), /* Out */ Triangle, /* Out */ TriangleVertexWeights ) )
 					{
 						const FVector SplitDirection = ( LaserImpactOnPolygonPlane - SplitStartLocation ).GetSafeNormal();
 
@@ -170,7 +170,7 @@ void USplitPolygonCommand::ApplyDuringDrag( IMeshEditorModeEditingContract& Mesh
 						FEdgeID ClosestEdgeID = FEdgeID::Invalid;
 
 						static TArray<FEdgeID> PolygonPerimeterEdgeIDs;
-						EditableMesh->GetPolygonPerimeterEdges( CandidatePolygonRef, /* Out */ PolygonPerimeterEdgeIDs );
+						EditableMesh->GetPolygonPerimeterEdges( CandidatePolygonID, /* Out */ PolygonPerimeterEdgeIDs );
 
 						for( const FEdgeID TargetEdgeID : PolygonPerimeterEdgeIDs )
 						{
@@ -312,7 +312,7 @@ void USplitPolygonCommand::ApplyDuringDrag( IMeshEditorModeEditingContract& Mesh
 									if( !bIsDisqualified )
 									{
 										// Connect to this vertex!
-										PolygonToSplit = CandidatePolygonRef;
+										PolygonToSplit = CandidatePolygonID;
 										ToVertexID = TargetVertexID;
 									}
 								}
@@ -322,7 +322,7 @@ void USplitPolygonCommand::ApplyDuringDrag( IMeshEditorModeEditingContract& Mesh
 							if( ToVertexID == FVertexID::Invalid )
 							{
 								// Split the edge to create a new vertex that we'll connect to
-								PolygonToSplit = CandidatePolygonRef;
+								PolygonToSplit = CandidatePolygonID;
 
 								// Go ahead and split the target edge
 								static TArray<FVertexID> NewVertexIDs;
@@ -345,7 +345,7 @@ void USplitPolygonCommand::ApplyDuringDrag( IMeshEditorModeEditingContract& Mesh
 		}
 
 
-		if( PolygonToSplit != FPolygonRef::Invalid )
+		if( PolygonToSplit != FPolygonID::Invalid )
 		{
 			check( ToVertexID != FVertexID::Invalid );
 
@@ -375,7 +375,7 @@ void USplitPolygonCommand::ApplyDuringDrag( IMeshEditorModeEditingContract& Mesh
 			PolygonsToSplit.Reset();
 			PolygonsToSplit.SetNum( 1, false );
 
-			PolygonsToSplit[ 0 ].PolygonRef = PolygonToSplit;
+			PolygonsToSplit[ 0 ].PolygonID = PolygonToSplit;
 			FVertexPair& VertexPair = *new( PolygonsToSplit[ 0 ].VertexPairsToSplitAt ) FVertexPair();
 			VertexPair.VertexID0 = FromVertexID;
 			VertexPair.VertexID1 = ToVertexID;
diff --git a/Engine/Plugins/Editor/MeshEditor/Source/PolygonModeling/TessellatePolygon.cpp b/Engine/Plugins/Editor/MeshEditor/Source/PolygonModeling/TessellatePolygon.cpp
--- a/Engine/Plugins/Editor/MeshEditor/Source/PolygonModeling/TessellatePolygon.cpp
+++ b/Engine/Plugins/Editor/MeshEditor/Source/PolygonModeling/TessellatePolygon.cpp
@@ -5,6 +5,7 @@
 #include "UICommandInfo.h"
 #include "EditableMesh.h"
 #include "MeshElement.h"
+#include "ScopedTransaction.h"
 
 
 #define LOCTEXT_NAMESPACE "MeshEditorMode"
@@ -48,13 +49,13 @@ void UTessellatePolygonCommand::Execute( IMeshEditorModeEditingContract& MeshEdi
 		UEditableMesh* EditableMesh = SelectedMeshAndPolygons.Key;
 		const TArray<FMeshElement>& PolygonElements = SelectedMeshAndPolygons.Value;
 
-		static TArray<FPolygonRef> PolygonsToTessellate;
+		static TArray<FPolygonID> PolygonsToTessellate;
 		PolygonsToTessellate.Reset( PolygonElements.Num() );
 
 		for( const auto& PolygonElement : PolygonElements )
 		{
-			const FPolygonRef PolygonRef( PolygonElement.ElementAddress.SectionID, FPolygonID( PolygonElement.ElementAddress.ElementID ) );
-			PolygonsToTessellate.Add( PolygonRef );
+			const FPolygonID PolygonID( PolygonElement.ElementAddress.ElementID );
+			PolygonsToTessellate.Add( PolygonID );
 		}
 
 
@@ -63,19 +64,18 @@ void UTessellatePolygonCommand::Execute( IMeshEditorModeEditingContract& MeshEdi
 		// @todo mesheditor: Expose as configurable parameter
 		const ETriangleTessellationMode TriangleTessellationMode = ETriangleTessellationMode::FourTriangles;
 
-		static TArray<FPolygonRef> NewPolygonRefs;
-		EditableMesh->TessellatePolygons( PolygonsToTessellate, TriangleTessellationMode, /* Out */ NewPolygonRefs );
+		static TArray<FPolygonID> NewPolygonIDs;
+		EditableMesh->TessellatePolygons( PolygonsToTessellate, TriangleTessellationMode, /* Out */ NewPolygonIDs );
 
 		// Select the new polygons
-		for( const FPolygonRef NewPolygonRef : NewPolygonRefs )
+		for( const FPolygonID NewPolygonID : NewPolygonIDs )
 		{
 			FMeshElement NewPolygonMeshElement;
 			{
 				NewPolygonMeshElement.Component = PolygonElements[0].Component;
 				NewPolygonMeshElement.ElementAddress = PolygonElements[0].ElementAddress;
 				NewPolygonMeshElement.ElementAddress.ElementType = EEditableMeshElementType::Polygon;
-				NewPolygonMeshElement.ElementAddress.SectionID = NewPolygonRef.SectionID;
-				NewPolygonMeshElement.ElementAddress.ElementID = NewPolygonRef.PolygonID;
+				NewPolygonMeshElement.ElementAddress.ElementID = NewPolygonID;
 			}
 
 			MeshElementsToSelect.Add( NewPolygonMeshElement );
