- Added event on UEditableMesh which is broadcast when any ElementIDs have been remapped to a new value (for example, following a Compact).
- Made compaction an opt-in process. By default, a new EditableMesh will not be sporadically compacted.
- Added the appropriate handler to MeshEditorMode so that all cached ElementIDs are remapped accordingly when the EditableMesh changes them.

diff --git a/Engine/Plugins/Editor/MeshEditor/Source/MeshEditor/MeshEditorMode.cpp b/Engine/Plugins/Editor/MeshEditor/Source/MeshEditor/MeshEditorMode.cpp
--- a/Engine/Plugins/Editor/MeshEditor/Source/MeshEditor/MeshEditorMode.cpp
+++ b/Engine/Plugins/Editor/MeshEditor/Source/MeshEditor/MeshEditorMode.cpp
@@ -461,6 +461,12 @@ FMeshEditorMode::~FMeshEditorMode()
 	FMeshEditorAnyElementCommands::Unregister();
 	FMeshEditorCommonCommands::Unregister();
 
+	// Remove the event registered on all cached editable meshes
+	for( auto& CachedEditableMesh : CachedEditableMeshes )
+	{
+		CachedEditableMesh.Value->OnElementIDsRemapped().RemoveAll( this );
+	}
+
 	MeshEditorModeProxyObject = nullptr;
 	AssetContainer = nullptr;
 }
@@ -484,6 +490,73 @@ void FMeshEditorMode::OnEndPIE( bool bIsSimulating )
 }
 
 
+void FMeshEditorMode::OnEditableMeshElementIDsRemapped( UEditableMesh* EditableMesh, const FElementIDRemappings& Remappings )
+{
+	// Helper function which performs the remapping of a given FMeshElement
+	auto RemapMeshElement = [ this, EditableMesh, &Remappings ]( FMeshElement& MeshElement )
+	{
+		UEditableMesh* MeshElementEditableMesh = this->FindOrCreateEditableMesh( *MeshElement.Component, MeshElement.ElementAddress.SubMeshAddress );
+		if( MeshElementEditableMesh == EditableMesh )
+		{
+			switch( MeshElement.ElementAddress.ElementType )
+			{
+			case EEditableMeshElementType::Vertex:
+				MeshElement.ElementAddress.ElementID = Remappings.GetRemappedVertexID( FVertexID( MeshElement.ElementAddress.ElementID ) );
+				break;
+
+			case EEditableMeshElementType::Edge:
+				MeshElement.ElementAddress.ElementID = Remappings.GetRemappedEdgeID( FEdgeID( MeshElement.ElementAddress.ElementID ) );
+				break;
+
+			case EEditableMeshElementType::Polygon:
+				MeshElement.ElementAddress.ElementID = Remappings.GetRemappedPolygonID( FPolygonID( MeshElement.ElementAddress.ElementID ) );
+				break;
+			}
+		}
+	};
+
+	// Change selected elements through the undo system so that they are restored correctly on an undo
+	FSelectOrDeselectMeshElementsChangeInput ChangeInput;
+	for( const FMeshElement& MeshElement : SelectedMeshElements )
+	{
+		ChangeInput.MeshElementsToDeselect.Add( MeshElement );
+		FMeshElement NewMeshElement = MeshElement;
+		RemapMeshElement( NewMeshElement );
+		ChangeInput.MeshElementsToSelect.Add( NewMeshElement );
+	}
+	TrackUndo( MeshEditorModeProxyObject, FSelectOrDeselectMeshElementsChange( MoveTemp( ChangeInput ) ).Execute( MeshEditorModeProxyObject ) );
+
+	for( FMeshElement& SelectedVertex : SelectedVertices )
+	{
+		check( SelectedVertex.ElementAddress.ElementType == EEditableMeshElementType::Vertex );
+		RemapMeshElement( SelectedVertex );
+	}
+
+	for( FMeshElement& SelectedEdge : SelectedEdges )
+	{
+		check( SelectedEdge.ElementAddress.ElementType == EEditableMeshElementType::Edge );
+		RemapMeshElement( SelectedEdge );
+	}
+
+	for( FMeshElement& SelectedPolygon : SelectedPolygons )
+	{
+		check( SelectedPolygon.ElementAddress.ElementType == EEditableMeshElementType::Polygon );
+		RemapMeshElement( SelectedPolygon );
+	}
+
+	for( FMeshElement& FadingOutHoveredMeshElement : FadingOutHoveredMeshElements )
+	{
+		RemapMeshElement( FadingOutHoveredMeshElement );
+	}
+
+	for( FMeshEditorInteractorData& MeshEditorInteractorData : MeshEditorInteractorDatas )
+	{
+		RemapMeshElement( MeshEditorInteractorData.HoveredMeshElement );
+		RemapMeshElement( MeshEditorInteractorData.PreviouslyHoveredMeshElement );
+	}
+}
+
+
 void FMeshEditorMode::RemoveEditableMeshReferences()
 {
 	// Instanced meshes live within the level itself. So remove all possible references to any editable mesh when the map changes,
@@ -866,6 +939,10 @@ UEditableMesh* FMeshEditorMode::FindOrCreateEditableMesh( UPrimitiveComponent& C
 
 			// Enable undo tracking on this mesh
 			EditableMesh->SetAllowUndo( true );
+
+			// Enable compaction on this mesh and set a callback so any cached ElementIDs can be remapped
+			EditableMesh->SetAllowCompact( true );
+			EditableMesh->OnElementIDsRemapped().AddRaw( this, &FMeshEditorMode::OnEditableMeshElementIDsRemapped );
 		}
 	}
 
diff --git a/Engine/Plugins/Editor/MeshEditor/Source/MeshEditor/MeshEditorMode.h b/Engine/Plugins/Editor/MeshEditor/Source/MeshEditor/MeshEditorMode.h
--- a/Engine/Plugins/Editor/MeshEditor/Source/MeshEditor/MeshEditorMode.h
+++ b/Engine/Plugins/Editor/MeshEditor/Source/MeshEditor/MeshEditorMode.h
@@ -342,6 +342,9 @@ protected:
 	/** Callback from the level editor when the map changes */
 	void OnMapChanged( UWorld* World, EMapChangeType MapChangeType );
 
+	/** Callback from an editable mesh when its elements are remapped */
+	void OnEditableMeshElementIDsRemapped( UEditableMesh* EditableMesh, const FElementIDRemappings& Remappings );
+
 	/** Callback from the level editor when new actors become selected or deselected */
 	void OnActorSelectionChanged( const TArray<UObject*>& NewSelection, bool bForceRefresh );
 
diff --git a/Engine/Plugins/Runtime/EditableMesh/Source/EditableMesh/EditableMesh.cpp b/Engine/Plugins/Runtime/EditableMesh/Source/EditableMesh/EditableMesh.cpp
--- a/Engine/Plugins/Runtime/EditableMesh/Source/EditableMesh/EditableMesh.cpp
+++ b/Engine/Plugins/Runtime/EditableMesh/Source/EditableMesh/EditableMesh.cpp
@@ -68,7 +68,9 @@ const FName UEditableMeshAttribute::PolygonCenterName( "PolygonCenter" );
 
 
 UEditableMesh::UEditableMesh()
-	: PendingCompactCounter( 0 )
+	: bAllowUndo( false ),
+	  bAllowCompact( false ),
+	  PendingCompactCounter( 0 )
 {
 }
 
@@ -280,7 +282,8 @@ void UEditableMesh::Compact()
 		Adapter->OnReindexElements( this, Remappings );
 	}
 
-	// @todo mesheditor: broadcast event with remappings so that any cached element IDs can be fixed up.
+	// Broadcast event with remappings so that any cached element IDs can be fixed up.
+	ElementIDsRemappedEvent.Broadcast( this, Remappings );
 
 	RebuildRenderMesh();
 
@@ -311,7 +314,8 @@ void UEditableMesh::Uncompact( const FElementIDRemappings& Remappings )
 		Adapter->OnReindexElements( this, Remappings );
 	}
 
-	// @todo mesheditor: broadcast event with remappings so that any cached element IDs can be fixed up.
+	// Broadcast event with remappings so that any cached element IDs can be fixed up.
+	ElementIDsRemappedEvent.Broadcast( this, Remappings );
 
 	RebuildRenderMesh();
 
@@ -430,7 +434,8 @@ void UEditableMesh::EndModification( const bool bFromUndo )
 
 		if( CurrentModificationType == EMeshModificationType::Final &&
 			CurrentToplogyChange == EMeshTopologyChange::TopologyChange &&
-			!bFromUndo )
+			!bFromUndo &&
+			bAllowCompact )
 		{
 			if( ++PendingCompactCounter == CompactFrequency )
 			{
diff --git a/Engine/Plugins/Runtime/EditableMesh/Source/EditableMesh/Public/EditableMesh.h b/Engine/Plugins/Runtime/EditableMesh/Source/EditableMesh/Public/EditableMesh.h
--- a/Engine/Plugins/Runtime/EditableMesh/Source/EditableMesh/Public/EditableMesh.h
+++ b/Engine/Plugins/Runtime/EditableMesh/Source/EditableMesh/Public/EditableMesh.h
@@ -609,6 +609,33 @@ public:
 	bool AnyChangesToUndo() const;
 
 
+	/**
+	 * @return	Returns true if compaction is enabled on this mesh
+	 */
+	UFUNCTION( BlueprintPure, Category="Editable Mesh" )
+	bool IsCompactAllowed() const
+	{
+		return bAllowCompact;
+	}
+
+	/**
+	 * Sets whether the mesh can be sporadically compacted as modifications are performed
+	 *
+	 * @param	bInAllowCompact		True if compaction is enabled on this mesh.
+	 */
+	UFUNCTION( BlueprintCallable, Category="Editable Mesh" )
+	void SetAllowCompact( const bool bInAllowCompact )
+	{
+		bAllowCompact = bInAllowCompact;
+	}
+
+
+	DECLARE_EVENT_TwoParams( UEditableMesh, FElementIDsRemapped, UEditableMesh*, const FElementIDRemappings& );
+	FElementIDsRemapped& OnElementIDsRemapped()
+	{
+		return ElementIDsRemappedEvent;
+	}
+
 	/**
 	 * Gets the sub-mesh address for this mesh which uniquely identifies the mesh among other sub-meshes in the same component
 	 *
@@ -814,6 +841,9 @@ public:
 	/** True if undo features are enabled on this mesh.  You're only allowed to call MakeUndo() if this is set to true. */
 	bool bAllowUndo;
 
+	/** True if compact is enabled on this mesh. If true, the mesh description will be sporadically compacted and tidied up. */
+	bool bAllowCompact;
+
 	/** When bAllowUndo is enabled, this will store the changes that can be applied to revert anything that happened to this
 	    mesh since the last time that MakeUndo() was called. */
 	TUniquePtr<FCompoundChangeInput> Undo;
@@ -850,7 +880,7 @@ public:
 	int32 PendingCompactCounter;
 
 	/** Data will be compacted after this many topology modifying actions. */
-	static const int32 CompactFrequency = 10;
+	static const int32 CompactFrequency = 1;
 
 	/** OpenSubdiv topology refiner object.  This is generated for meshes that have subdivision levels, and reused to generate new limit surfaces 
 	    when geometry is moved.  When the mesh's topology changes, this object is regenerated from scratch. */
@@ -874,4 +904,7 @@ public:
 
 	/** The resulting limit surface geometry after GenerateOpenSubdivLimitSurfaceData() is called */
 	FSubdivisionLimitData SubdivisionLimitData;
+
+	/** Broadcast event when element IDs are remapped (for example, following Compact / Uncompact) */
+	FElementIDsRemapped ElementIDsRemappedEvent;
 };
