Selected ElementID remapping no longer done through an FChange. It doesn't need to be undoable, as the undo compact will broadcast the ElementIDsRemapped event anyway.

diff --git a/Engine/Plugins/Editor/MeshEditor/Source/MeshEditor/MeshEditorMode.cpp b/Engine/Plugins/Editor/MeshEditor/Source/MeshEditor/MeshEditorMode.cpp
--- a/Engine/Plugins/Editor/MeshEditor/Source/MeshEditor/MeshEditorMode.cpp
+++ b/Engine/Plugins/Editor/MeshEditor/Source/MeshEditor/MeshEditorMode.cpp
@@ -515,16 +515,10 @@ void FMeshEditorMode::OnEditableMeshElementIDsRemapped( UEditableMesh* EditableM
 		}
 	};
 
-	// Change selected elements through the undo system so that they are restored correctly on an undo
-	FSelectOrDeselectMeshElementsChangeInput ChangeInput;
-	for( const FMeshElement& MeshElement : SelectedMeshElements )
+	for( FMeshElement& MeshElement : SelectedMeshElements )
 	{
-		ChangeInput.MeshElementsToDeselect.Add( MeshElement );
-		FMeshElement NewMeshElement = MeshElement;
-		RemapMeshElement( NewMeshElement );
-		ChangeInput.MeshElementsToSelect.Add( NewMeshElement );
+		RemapMeshElement( MeshElement );
 	}
-	TrackUndo( MeshEditorModeProxyObject, FSelectOrDeselectMeshElementsChange( MoveTemp( ChangeInput ) ).Execute( MeshEditorModeProxyObject ) );
 
 	for( FMeshElement& SelectedVertex : SelectedVertices )
 	{
