Mesh Editor: Fast element hover and selection; Spatial database for editable meshes

- Editable meshes now generate an octree for themselves in the editor
- Interactor hover and selection queries have been changed to use the octree for super-fast performance
- The octree is incrementally updated as the mesh is modified.  There is some additional overhead for this when performing changes to the mesh.
- The octree is generated on initialization.  There is some additional cost (hitch) when this happens.  The octree is never saved or loaded.
- New cvar:  "EditableMesh.OctreeIncrementalUpdateLimit".  Controls how much geometry may change before we just force a full rebuild of the octree.
- New cvar: "EditableMesh.UseBoundlessOctree".  Can be enabled to use a tight octree bounds.  Only valid for debugging.

Other changes:
- New cvar:  'MeshEd.ShowWiresForSelectedMeshes'.  It can be set to zero to hide mesh wire frames.
- Fixed instanced copies of editable meshes not being eligble for compacting

Known Issues:  (search for '[REDACTED_HANDLE] mesheditor spatial')
- 'Loose' edges and vertices (not connected to any polygons) aren't selectable.
- Marquee select is not accelerated yet.
- Grabber sphere mode (not supported yet) is not accelerated yet.
- 'Fuzzy' selection may sometimes not catch mesh elements that are not directly under the interactor and very near an octree node bounding box

diff --git a/Engine/Plugins/Editor/MeshEditor/Source/MeshEditor/MeshEditorMode.cpp b/Engine/Plugins/Editor/MeshEditor/Source/MeshEditor/MeshEditorMode.cpp
--- a/Engine/Plugins/Editor/MeshEditor/Source/MeshEditor/MeshEditorMode.cpp
+++ b/Engine/Plugins/Editor/MeshEditor/Source/MeshEditor/MeshEditorMode.cpp
@@ -70,6 +70,7 @@ namespace MeshEd
 	static FAutoConsoleVariable MinDeltaForInertialMovement( TEXT( "MeshEd.MinDeltaForInertialMovement" ), 0.01f, TEXT( "Minimum velocity in cm/frame for inertial movement to kick in when releasing a drag" ) );
 	static FAutoConsoleVariable ShowDebugStats( TEXT( "MeshEd.ShowDebugStats" ), 0, TEXT( "Enables debug overlay text for the currently selected mesh" ) );
 	static FAutoConsoleVariable EnableSelectByPainting( TEXT( "MeshEd.EnableSelectByPainting" ), 0, TEXT( "Enables selection by clicking and dragging over elements" ) );
+	static FAutoConsoleVariable ShowWiresForSelectedMeshes( TEXT( "MeshEd.ShowWiresForSelectedMeshes" ), 1, TEXT( "Enables rendering of a polygonal wireframe overlay on selected meshes" ) );
 
 	static FAutoConsoleVariable OverlayDepthOffset( TEXT( "MeshEd.OverlayDepthOffset" ), 2.0f, TEXT( "How far to offset overlay wires/polygons on top of meshes when hovered or selected" ) );
 	static FAutoConsoleVariable OverlayVertexSize( TEXT( "MeshEd.OverlayVertexSize" ), 4.0f, TEXT( "How large a vertex is on a mesh overlay" ) );
@@ -934,6 +935,9 @@ UEditableMesh* FMeshEditorMode::FindOrCreateEditableMesh( UPrimitiveComponent& C
 			// Enable undo tracking on this mesh
 			EditableMesh->SetAllowUndo( true );
 
+			// Enable spatial database, so that we can quickly query which polygons are under the mouse cursor
+			EditableMesh->SetAllowSpatialDatabase( true );
+
 			// Enable compaction on this mesh and set a callback so any cached ElementIDs can be remapped
 			EditableMesh->SetAllowCompact( true );
 			EditableMesh->OnElementIDsRemapped().AddRaw( this, &FMeshEditorMode::OnEditableMeshElementIDsRemapped );
@@ -1169,6 +1173,8 @@ void FMeshEditorMode::CommitEditableMeshIfNecessary( UEditableMesh* EditableMesh
 	{
 		UEditableMesh* NewEditableMesh = EditableMesh->CommitInstance( Component );
 		NewEditableMesh->SetAllowUndo(true);
+		NewEditableMesh->SetAllowSpatialDatabase(true);
+		NewEditableMesh->SetAllowCompact(true);
 
 		// Commit the editable mesh as a new instance in the static mesh component
 		const FEditableMeshSubMeshAddress& OldSubMeshAddress = EditableMesh->GetSubMeshAddress();
@@ -2351,29 +2357,9 @@ void FMeshEditorMode::Render( const FSceneView* SceneView, FViewport* Viewport,
 	}
 
 	// Draw all polygon edges for selected/hovered meshes (@todo mesheditor: This only draws polygons right now.  We should also draw "floating" edges and vertices too!)
+	if( MeshEd::ShowWiresForSelectedMeshes->GetInt() != 0 )
 	{
-		static TArray<FComponentAndEditableMesh> HoveredOrSelectedMeshes;
-		HoveredOrSelectedMeshes.Reset();
-
-		HoveredOrSelectedMeshes = SelectedComponentsAndEditableMeshes;
-
-		// Only draw hover if we're not in the middle of an interactive edit
-		if( ActiveAction == NAME_None )
-		{
-			for( FMeshEditorInteractorData& MeshEditorInteractorData : MeshEditorInteractorDatas )
-			{
-				if( MeshEditorInteractorData.HoveredMeshElement.IsValidMeshElement() )
-				{
-					UEditableMesh* EditableMesh = FindOrCreateEditableMesh( *MeshEditorInteractorData.HoveredMeshElement.Component, MeshEditorInteractorData.HoveredMeshElement.ElementAddress.SubMeshAddress );
-					if( EditableMesh != nullptr )
-					{
-						HoveredOrSelectedMeshes.AddUnique( FComponentAndEditableMesh( MeshEditorInteractorData.HoveredMeshElement.Component.Get(), EditableMesh ) );
-					}
-				}
-			}
-		}
-
-		for( const FComponentAndEditableMesh& ComponentAndEditableMesh : HoveredOrSelectedMeshes )
+		for( const FComponentAndEditableMesh& ComponentAndEditableMesh : SelectedComponentsAndEditableMeshes )
 		{
 			const UPrimitiveComponent* Component = ComponentAndEditableMesh.Component;
 			check( Component != nullptr );
@@ -2747,11 +2733,10 @@ void FMeshEditorMode::OnViewportInteractionHoverUpdate( UViewportInteractor* Vie
 				const int32 FirstInteractorPassNumber = GetDefault<UMeshEditorSettings>()->bAllowGrabberSphere ? 0 : 1;
 				for( int32 InteractorPassNumber = FirstInteractorPassNumber; InteractorPassNumber < 2; ++InteractorPassNumber )
 				{
-					const bool bIsGrabberSphereTest = ( InteractorPassNumber == 0 );
-					const bool bIsLaserTest = !bIsGrabberSphereTest;
+					const EInteractorShape InteractorShape = ( InteractorPassNumber == 0 ) ? EInteractorShape::GrabberSphere : EInteractorShape::Laser;
 
-					if( ( bIsGrabberSphereTest && MeshEditorInteractorData.bGrabberSphereIsValid ) ||
-						( bIsLaserTest && MeshEditorInteractorData.bLaserIsValid ) )
+					if( ( InteractorShape == EInteractorShape::GrabberSphere && MeshEditorInteractorData.bGrabberSphereIsValid ) ||
+						( InteractorShape == EInteractorShape::Laser && MeshEditorInteractorData.bLaserIsValid ) )
 					{
 						static TArray< UPrimitiveComponent* > HitComponents;
 						HitComponents.Reset();
@@ -2767,7 +2752,7 @@ void FMeshEditorMode::OnViewportInteractionHoverUpdate( UViewportInteractor* Vie
 							static TArray<UPrimitiveComponent*> ComponentsFoundThisPass;
 							ComponentsFoundThisPass.Reset();
 
-							if( bIsGrabberSphereTest )
+							if( InteractorShape == EInteractorShape::GrabberSphere )
 							{
 								// Grabber sphere testing
 								FCollisionShape CollisionShape;
@@ -2795,8 +2780,8 @@ void FMeshEditorMode::OnViewportInteractionHoverUpdate( UViewportInteractor* Vie
 								FCollisionShape CollisionShape;
 								CollisionShape.SetSphere( WorldSpaceRayFuzzyDistance * ExtraFuzzyScalingForCollisionQuery );
 
-								// @todo mesheditor perf: This could be fairly slow, tracing so many objects.  We could do SweepSingleByChannel, but the nearest mesh might not
-								// actually have the best element to select
+								// @todo mesheditor perf: This could be fairly slow, tracing so many objects.  We could do SweepSingleByChannel, but the nearest mesh might not actually have the best element to select
+								// @todo mesheditor perf: Do we really need to even do a complex PhysX trace now that we have spatial databases for editable meshes?
 								static TArray< FHitResult > HitResults;
 								HitResults.Reset();
 								if( GetWorld()->SweepMultiByChannel( HitResults, MeshEditorInteractorData.LaserStart, MeshEditorInteractorData.LaserEnd, FQuat::Identity, ECC_Visibility, CollisionShape, TraceParams ) )
@@ -2891,10 +2876,9 @@ void FMeshEditorMode::OnViewportInteractionHoverUpdate( UViewportInteractor* Vie
 									FVector ComponentSpaceHitLocation = FVector::ZeroVector;
 									FEditableMeshElementAddress MeshElementAddress = QueryElement( 
 										*EditableMesh, 
-										bIsGrabberSphereTest, 
+										InteractorShape,
 										ComponentSpaceGrabberSphere, 
 										ComponentSpaceGrabberSphereFuzzyDistance, 
-										bIsLaserTest, 
 										ComponentSpaceLaserStart, 
 										ComponentSpaceLaserEnd, 
 										ComponentSpaceRayFuzzyDistance, 
@@ -2944,7 +2928,7 @@ void FMeshEditorMode::OnViewportInteractionHoverUpdate( UViewportInteractor* Vie
 
 							// If we hit something with our grabber sphere, then don't bother checking with the laser.  We always
 							// prefer grabber sphere hits.
-							if( bIsGrabberSphereTest )
+							if( InteractorShape == EInteractorShape::GrabberSphere )
 							{
 								//	DrawDebugSphere( GetWorld(), MeshEditorInteractorData.HoverLocation, 1.5f * ViewportWorldInteraction->GetWorldScaleFactor(), 16, FColor( 255, 40, 40, 255 ), false, 0.0f );
 								bIsGrabberSphereOverMeshElement = true;
@@ -3654,13 +3638,46 @@ bool FMeshEditorMode::FindEdgeSplitUnderInteractor(	UViewportInteractor* Viewpor
 }
 
 
-FEditableMeshElementAddress FMeshEditorMode::QueryElement( const UEditableMesh& EditableMesh, const bool bUseSphere, const FSphere& Sphere, const float SphereFuzzyDistance, bool bUseRay, const FVector& RayStart, const FVector& RayEnd, const float RayFuzzyDistance, const EEditableMeshElementType OnlyElementType, const FVector& CameraLocation, const bool bIsPerspectiveView, const float FuzzyDistanceScaleFactor, EInteractorShape& OutInteractorShape, FVector& OutHitLocation ) const
+FEditableMeshElementAddress FMeshEditorMode::QueryElement( const UEditableMesh& EditableMesh, const EInteractorShape InteractorShape, const FSphere& Sphere, const float SphereFuzzyDistance, const FVector& RayStart, const FVector& RayEnd, const float RayFuzzyDistance, const EEditableMeshElementType OnlyElementType, const FVector& CameraLocation, const bool bIsPerspectiveView, const float FuzzyDistanceScaleFactor, EInteractorShape& OutInteractorShape, FVector& OutHitLocation ) const
 {
 	OutHitLocation = FVector::ZeroVector;
 
 	FEditableMeshElementAddress HitElementAddress;
 	HitElementAddress.SubMeshAddress = EditableMesh.GetSubMeshAddress();
 
+
+	// Figure out our candidate set of polygons by performing a spatial query on the mesh
+	static TArray<FPolygonID> CandidatePolygons;
+	CandidatePolygons.Reset();
+
+	if( InteractorShape == EInteractorShape::Laser )
+	{
+		// @todo mesheditor spatial: Do we need to use a "fat ray" to account for fuzzy testing (or expanded octree boxes)?  We don't currently have a 'segment distance to AABB' function.
+
+		check( EditableMesh.IsSpatialDatabaseAllowed() );	// We need a spatial database to do this query fast!
+		EditableMesh.SearchSpatialDatabaseForPolygonsPotentiallyIntersectingLineSegment( RayStart, RayEnd, /* Out */ CandidatePolygons );
+
+		// @todo mesheditor debug
+		// 	if( EditableMesh.GetPolygonCount() > 0 )
+		// 	{
+		// 		UE_LOG( LogEditableMesh, Display, TEXT( "%i  (%0.1f%%)" ), CandidatePolygons.Num(), ( (float)CandidatePolygons.Num() / (float)EditableMesh.GetPolygonCount() ) * 100.0f );
+		// 	}
+	}
+	else
+	{
+		// @todo mesheditor spatial: Need GrabberSphere support for spatial queries.  Currently we're just testing all polygons (slow!)
+		const uint32 PolygonArraySize = EditableMesh.GetPolygonArraySize();
+		for( uint32 PolygonIndex = 0; PolygonIndex < PolygonArraySize; ++PolygonIndex )
+		{
+			const FPolygonID PolygonID( PolygonIndex );
+
+			if( EditableMesh.IsValidPolygon( PolygonID ) )
+			{
+				CandidatePolygons.Add( PolygonID );
+			}
+		}
+	}
+
 	static TSet<FVertexID> FrontFacingVertices;
 	FrontFacingVertices.Reset();
 	
@@ -3671,26 +3688,22 @@ FEditableMeshElementAddress FMeshEditorMode::QueryElement( const UEditableMesh&
 	FrontFacingPolygons.Reset();
 
 	// Look for all the front-facing elements
-	const uint32 PolygonArraySize = EditableMesh.GetPolygonArraySize();
-	for( uint32 PolygonIndex = 0; PolygonIndex < PolygonArraySize; ++PolygonIndex )
+	for( const FPolygonID PolygonID : CandidatePolygons )
 	{
-		const FPolygonID PolygonID( PolygonIndex );
-
-		if( EditableMesh.IsValidPolygon( PolygonID ) )
+		const FVector PolygonNormal = EditableMesh.ComputePolygonNormal( PolygonID );
+		const FVector PolygonCenter = EditableMesh.ComputePolygonCenter( PolygonID );
+		if( ( InteractorShape == EInteractorShape::GrabberSphere ) ||	// Sphere tests never eliminate back-facing geometry
+			!bIsPerspectiveView ||			// @todo mesheditor: Add support for backface culling in orthographic views
+			FVector::DotProduct( CameraLocation - PolygonCenter, PolygonNormal ) > 0.0f )
 		{
-			const FVector PolygonNormal = EditableMesh.ComputePolygonNormal( PolygonID );
-			const FVector PolygonCenter = EditableMesh.ComputePolygonCenter( PolygonID );
-			if( !bIsPerspectiveView || FVector::DotProduct( CameraLocation - PolygonCenter, PolygonNormal ) > 0.0f )	// @todo mesheditor: Add support for backface culling in orthographic views
-			{
-				FrontFacingPolygons.Add( PolygonID );
+			FrontFacingPolygons.Add( PolygonID );
 
-				const int32 PolygonVertexCount = EditableMesh.GetPolygonPerimeterVertexCount( PolygonID );
-				for( int32 Index = 0; Index < PolygonVertexCount; ++Index )
-				{
-					FrontFacingVertices.Add( EditableMesh.GetPolygonPerimeterVertex( PolygonID, Index ) );
-					bool bOutEdgeWindingIsReversedForPolygons;
-					FrontFacingEdges.Add( EditableMesh.GetPolygonPerimeterEdge( PolygonID, Index, bOutEdgeWindingIsReversedForPolygons ) );
-				}
+			const int32 PolygonVertexCount = EditableMesh.GetPolygonPerimeterVertexCount( PolygonID );
+			for( int32 Index = 0; Index < PolygonVertexCount; ++Index )
+			{
+				FrontFacingVertices.Add( EditableMesh.GetPolygonPerimeterVertex( PolygonID, Index ) );
+				bool bOutEdgeWindingIsReversedForPolygons;
+				FrontFacingEdges.Add( EditableMesh.GetPolygonPerimeterEdge( PolygonID, Index, bOutEdgeWindingIsReversedForPolygons ) );
 			}
 		}
 	}
@@ -3719,7 +3732,7 @@ FEditableMeshElementAddress FMeshEditorMode::QueryElement( const UEditableMesh&
 			}
 
 			const bool bAlreadyHitTriangle = ( HitElementAddress.ElementType == EEditableMeshElementType::Polygon );
-			const bool bHit = CheckTriangle( bUseSphere, Sphere, SphereFuzzyDistance, bUseRay, RayStart, CurrentRayEnd, RayFuzzyDistance, TriangleVertexPositions, CameraLocation, bIsPerspectiveView, FuzzyDistanceScaleFactor, ClosestInteractorShape, ClosestDistanceToRay, ClosestDistanceOnRay, ClosestHitLocation, bAlreadyHitTriangle );
+			const bool bHit = CheckTriangle( InteractorShape, Sphere, SphereFuzzyDistance, RayStart, CurrentRayEnd, RayFuzzyDistance, TriangleVertexPositions, CameraLocation, bIsPerspectiveView, FuzzyDistanceScaleFactor, ClosestInteractorShape, ClosestDistanceToRay, ClosestDistanceOnRay, ClosestHitLocation, bAlreadyHitTriangle );
 			if( bHit )
 			{
 				HitElementAddress.ElementType = EEditableMeshElementType::Polygon;
@@ -3741,7 +3754,7 @@ FEditableMeshElementAddress FMeshEditorMode::QueryElement( const UEditableMesh&
 			EdgeVertexPositions[ 1 ] = EditableMesh.GetVertexAttribute( EditableMesh.GetEdgeVertex( EdgeID, 1 ), UEditableMeshAttribute::VertexPosition(), 0 );
 
 			const bool bAlreadyHitEdge = ( HitElementAddress.ElementType == EEditableMeshElementType::Edge );
-			const bool bHit = CheckEdge( bUseSphere, Sphere, SphereFuzzyDistance, bUseRay, RayStart, CurrentRayEnd, RayFuzzyDistance, EdgeVertexPositions, CameraLocation, bIsPerspectiveView, FuzzyDistanceScaleFactor, ClosestInteractorShape, ClosestDistanceToRay, ClosestDistanceOnRay, ClosestHitLocation, bAlreadyHitEdge );
+			const bool bHit = CheckEdge( InteractorShape, Sphere, SphereFuzzyDistance, RayStart, CurrentRayEnd, RayFuzzyDistance, EdgeVertexPositions, CameraLocation, bIsPerspectiveView, FuzzyDistanceScaleFactor, ClosestInteractorShape, ClosestDistanceToRay, ClosestDistanceOnRay, ClosestHitLocation, bAlreadyHitEdge );
 			if( bHit )
 			{
 				HitElementAddress.ElementType = EEditableMeshElementType::Edge;
@@ -3750,8 +3763,6 @@ FEditableMeshElementAddress FMeshEditorMode::QueryElement( const UEditableMesh&
 		}
 	}
 
-	// @todo mesheditor perf:  OMG too slow.  Needs acceleration structure + MT.  Horribly slow to iterate over all elements like this.  We need a proper spatial query!
-
 	ClosestDistanceToRay = TNumericLimits<float>::Max();
 
 	// Check vertices
@@ -3761,7 +3772,7 @@ FEditableMeshElementAddress FMeshEditorMode::QueryElement( const UEditableMesh&
 		{
 			const FVector VertexPosition = EditableMesh.GetVertexAttribute( VertexID, UEditableMeshAttribute::VertexPosition(), 0 );
 			const bool bAlreadyHitVertex = ( HitElementAddress.ElementType == EEditableMeshElementType::Vertex );
-			const bool bHit = CheckVertex( bUseSphere, Sphere, SphereFuzzyDistance, bUseRay, RayStart, CurrentRayEnd, RayFuzzyDistance, VertexPosition, CameraLocation, bIsPerspectiveView, FuzzyDistanceScaleFactor, ClosestInteractorShape, ClosestDistanceToRay, ClosestDistanceOnRay, ClosestHitLocation, bAlreadyHitVertex );
+			const bool bHit = CheckVertex( InteractorShape, Sphere, SphereFuzzyDistance, RayStart, CurrentRayEnd, RayFuzzyDistance, VertexPosition, CameraLocation, bIsPerspectiveView, FuzzyDistanceScaleFactor, ClosestInteractorShape, ClosestDistanceToRay, ClosestDistanceOnRay, ClosestHitLocation, bAlreadyHitVertex );
 			if( bHit )
 			{
 				HitElementAddress.ElementType = EEditableMeshElementType::Vertex;
@@ -3780,7 +3791,7 @@ FEditableMeshElementAddress FMeshEditorMode::QueryElement( const UEditableMesh&
 }
 
 
-bool FMeshEditorMode::CheckVertex( const bool bUseSphere, const FSphere& Sphere, const float SphereFuzzyDistance, bool bUseRay, const FVector& RayStart, const FVector& RayEnd, const float RayFuzzyDistance, const FVector& VertexPosition, const FVector& CameraLocation, const bool bIsPerspectiveView, const float FuzzyDistanceScaleFactor, EInteractorShape& ClosestInteractorShape, float& ClosestDistanceToRay, float& ClosestDistanceOnRay, FVector& ClosestHitLocation, const bool bAlreadyHitVertex )
+bool FMeshEditorMode::CheckVertex( const EInteractorShape InteractorShape, const FSphere& Sphere, const float SphereFuzzyDistance, const FVector& RayStart, const FVector& RayEnd, const float RayFuzzyDistance, const FVector& VertexPosition, const FVector& CameraLocation, const bool bIsPerspectiveView, const float FuzzyDistanceScaleFactor, EInteractorShape& ClosestInteractorShape, float& ClosestDistanceToRay, float& ClosestDistanceOnRay, FVector& ClosestHitLocation, const bool bAlreadyHitVertex )
 {
 	bool bHit = false;
 
@@ -3789,7 +3800,7 @@ bool FMeshEditorMode::CheckVertex( const bool bUseSphere, const FSphere& Sphere,
 	const float DistanceBasedScaling = DistanceBias + DistanceToCamera * FuzzyDistanceScaleFactor;
 	check( DistanceBasedScaling > 0.0f );
 
-	if( bUseSphere )
+	if( InteractorShape == EInteractorShape::GrabberSphere )
 	{
 		const float DistanceToSphere = ( VertexPosition - Sphere.Center ).Size();
 		if( DistanceToSphere <= Sphere.W )
@@ -3808,7 +3819,7 @@ bool FMeshEditorMode::CheckVertex( const bool bUseSphere, const FSphere& Sphere,
 		}
 	}
 
-	if( bUseRay )
+	if( InteractorShape == EInteractorShape::Laser )
 	{
 		const FVector ClosestPointOnRay = FMath::ClosestPointOnSegment( VertexPosition, RayStart, RayEnd );
 		const float DistanceToRay = ( ClosestPointOnRay - VertexPosition ).Size();
@@ -3844,11 +3855,11 @@ bool FMeshEditorMode::CheckVertex( const bool bUseSphere, const FSphere& Sphere,
 }
 
 
-bool FMeshEditorMode::CheckEdge( const bool bUseSphere, const FSphere& Sphere, const float SphereFuzzyDistance, bool bUseRay, const FVector& RayStart, const FVector& RayEnd, const float RayFuzzyDistance, const FVector EdgeVertexPositions[ 2 ], const FVector& CameraLocation, const bool bIsPerspectiveView, const float FuzzyDistanceScaleFactor, EInteractorShape& ClosestInteractorShape, float& ClosestDistanceToRay, float& ClosestDistanceOnRay, FVector& ClosestHitLocation, const bool bAlreadyHitEdge )
+bool FMeshEditorMode::CheckEdge( const EInteractorShape InteractorShape, const FSphere& Sphere, const float SphereFuzzyDistance, const FVector& RayStart, const FVector& RayEnd, const float RayFuzzyDistance, const FVector EdgeVertexPositions[ 2 ], const FVector& CameraLocation, const bool bIsPerspectiveView, const float FuzzyDistanceScaleFactor, EInteractorShape& ClosestInteractorShape, float& ClosestDistanceToRay, float& ClosestDistanceOnRay, FVector& ClosestHitLocation, const bool bAlreadyHitEdge )
 {
 	bool bHit = false;
 
-	if( bUseSphere )
+	if( InteractorShape == EInteractorShape::GrabberSphere )
 	{
 		const float DistanceToSphere = FMath::PointDistToSegment( Sphere.Center, EdgeVertexPositions[ 0 ], EdgeVertexPositions[ 1 ] );
 		if( DistanceToSphere <= Sphere.W )
@@ -3873,7 +3884,7 @@ bool FMeshEditorMode::CheckEdge( const bool bUseSphere, const FSphere& Sphere, c
 	}
 
 
-	if( bUseRay )
+	if( InteractorShape == EInteractorShape::Laser )
 	{
 		FVector ClosestPointOnEdge, ClosestPointOnRay;
 		FMath::SegmentDistToSegmentSafe(
@@ -3919,11 +3930,11 @@ bool FMeshEditorMode::CheckEdge( const bool bUseSphere, const FSphere& Sphere, c
 }
 
 
-bool FMeshEditorMode::CheckTriangle( const bool bUseSphere, const FSphere& Sphere, const float SphereFuzzyDistance, bool bUseRay, const FVector& RayStart, const FVector& RayEnd, const float RayFuzzyDistance, const FVector TriangleVertexPositions[ 3 ], const FVector& CameraLocation, const bool bIsPerspectiveView, const float FuzzyDistanceScaleFactor, EInteractorShape& ClosestInteractorShape, float& ClosestDistanceToRay, float& ClosestDistanceOnRay, FVector& ClosestHitLocation, const bool bAlreadyHitTriangle )
+bool FMeshEditorMode::CheckTriangle( const EInteractorShape InteractorShape, const FSphere& Sphere, const float SphereFuzzyDistance, const FVector& RayStart, const FVector& RayEnd, const float RayFuzzyDistance, const FVector TriangleVertexPositions[ 3 ], const FVector& CameraLocation, const bool bIsPerspectiveView, const float FuzzyDistanceScaleFactor, EInteractorShape& ClosestInteractorShape, float& ClosestDistanceToRay, float& ClosestDistanceOnRay, FVector& ClosestHitLocation, const bool bAlreadyHitTriangle )
 {
 	bool bHit = false;
 
-	if( bUseSphere )
+	if( InteractorShape == EInteractorShape::GrabberSphere )
 	{
 		// @todo grabber: FMath::ClosestPointOnTriangleToPoint doesn't work with degenerates (always returns ray start point?)
 		const FVector ClosestPointOnTriangleToSphere = FMath::ClosestPointOnTriangleToPoint( Sphere.Center, TriangleVertexPositions[ 0 ], TriangleVertexPositions[ 1 ], TriangleVertexPositions[ 2 ] );
@@ -3949,7 +3960,7 @@ bool FMeshEditorMode::CheckTriangle( const bool bUseSphere, const FSphere& Spher
 	}
 
 
-	if( bUseRay )
+	if( InteractorShape == EInteractorShape::Laser )
 	{
 		// @todo mesheditor: We have like 5 different versions of this in the engine, but nothing generic in a nice place
 		struct Local
@@ -4376,6 +4387,8 @@ void FMeshEditorMode::PostUndo()
 
 bool FMeshEditorMode::FrustumSelect( const FConvexVolume& InFrustum, FEditorViewportClient* InViewportClient, bool InSelect )
 {
+	// @todo mesheditor spatial: Need to update marquee select to use spatial queries
+
 	// @todo mesheditor urgent: settings class for bundling together all these kind of things
 	const bool bShouldDeselectAllFirst = true;	// @todo mesheditor: needs to be passed to this method
 	const bool bOnlySelectVisibleMeshes = GetDefault<UMeshEditorSettings>()->bOnlySelectVisibleMeshes;
diff --git a/Engine/Plugins/Editor/MeshEditor/Source/MeshEditor/MeshEditorMode.h b/Engine/Plugins/Editor/MeshEditor/Source/MeshEditor/MeshEditorMode.h
--- a/Engine/Plugins/Editor/MeshEditor/Source/MeshEditor/MeshEditorMode.h
+++ b/Engine/Plugins/Editor/MeshEditor/Source/MeshEditor/MeshEditorMode.h
@@ -252,10 +252,10 @@ protected:
 	void UpdateActiveAction( const bool bIsActionFinishing );
 
 	/** Geometry tests */
-	FEditableMeshElementAddress QueryElement( const UEditableMesh& EditableMesh, const bool bUseSphere, const FSphere& Sphere, const float SphereFuzzyDistance, bool bUseRay, const FVector& RayStart, const FVector& RayEnd, const float RayFuzzyDistance, const EEditableMeshElementType OnlyElementType, const FVector& CameraLocation, const bool bIsPerspectiveView, const float FuzzyDistanceScaleFactor, EInteractorShape& OutInteractorShape, FVector& OutHitLocation ) const;
-	static bool CheckVertex( const bool bUseSphere, const FSphere& Sphere, const float SphereFuzzyDistance, bool bUseRay, const FVector& RayStart, const FVector& RayEnd, const float FuzzyDistance, const FVector& VertexPosition, const FVector& CameraLocation, const bool bIsPerspectiveView, const float FuzzyDistanceScaleFactor, EInteractorShape& ClosestInteractorShape, float& ClosestDistanceToRay, float& ClosestDistanceOnRay, FVector& ClosestHitLocation, const bool bAlreadyHitVertex );
-	static bool CheckEdge( const bool bUseSphere, const FSphere& Sphere, const float SphereFuzzyDistance, bool bUseRay, const FVector& RayStart, const FVector& RayEnd, const float FuzzyDistance, const FVector EdgeVertexPositions[ 2 ], const FVector& CameraLocation, const bool bIsPerspectiveView, const float FuzzyDistanceScaleFactor, EInteractorShape& ClosestInteractorShape, float& ClosestDistanceToRay, float& ClosestDistanceOnRay, FVector& ClosestHitLocation, const bool bAlreadyEdge );
-	static bool CheckTriangle( const bool bUseSphere, const FSphere& Sphere, const float SphereFuzzyDistance, bool bUseRay, const FVector& RayStart, const FVector& RayEnd, const float FuzzyDistance, const FVector TriangleVertexPositions[ 3 ], const FVector& CameraLocation, const bool bIsPerspectiveView, const float FuzzyDistanceScaleFactor, EInteractorShape& ClosestInteractorShape, float& ClosestDistanceToRay, float& ClosestDistanceOnRay, FVector& ClosestHitLocation, const bool bAlreadyHitTriangle );
+	FEditableMeshElementAddress QueryElement( const UEditableMesh& EditableMesh, const EInteractorShape InteractorShape, const FSphere& Sphere, const float SphereFuzzyDistance, const FVector& RayStart, const FVector& RayEnd, const float RayFuzzyDistance, const EEditableMeshElementType OnlyElementType, const FVector& CameraLocation, const bool bIsPerspectiveView, const float FuzzyDistanceScaleFactor, EInteractorShape& OutInteractorShape, FVector& OutHitLocation ) const;
+	static bool CheckVertex( const EInteractorShape InteractorShape, const FSphere& Sphere, const float SphereFuzzyDistance, const FVector& RayStart, const FVector& RayEnd, const float FuzzyDistance, const FVector& VertexPosition, const FVector& CameraLocation, const bool bIsPerspectiveView, const float FuzzyDistanceScaleFactor, EInteractorShape& ClosestInteractorShape, float& ClosestDistanceToRay, float& ClosestDistanceOnRay, FVector& ClosestHitLocation, const bool bAlreadyHitVertex );
+	static bool CheckEdge( const EInteractorShape InteractorShape, const FSphere& Sphere, const float SphereFuzzyDistance, const FVector& RayStart, const FVector& RayEnd, const float FuzzyDistance, const FVector EdgeVertexPositions[ 2 ], const FVector& CameraLocation, const bool bIsPerspectiveView, const float FuzzyDistanceScaleFactor, EInteractorShape& ClosestInteractorShape, float& ClosestDistanceToRay, float& ClosestDistanceOnRay, FVector& ClosestHitLocation, const bool bAlreadyEdge );
+	static bool CheckTriangle( const EInteractorShape InteractorShape, const FSphere& Sphere, const float SphereFuzzyDistance, const FVector& RayStart, const FVector& RayEnd, const float FuzzyDistance, const FVector TriangleVertexPositions[ 3 ], const FVector& CameraLocation, const bool bIsPerspectiveView, const float FuzzyDistanceScaleFactor, EInteractorShape& ClosestInteractorShape, float& ClosestDistanceToRay, float& ClosestDistanceOnRay, FVector& ClosestHitLocation, const bool bAlreadyHitTriangle );
 
 	/** Returns the index of an element in the selection set, or INDEX_NONE if its not selected */
 	int32 GetSelectedMeshElementIndex( const FMeshElement& MeshElement ) const;
diff --git a/Engine/Plugins/Runtime/EditableMesh/Source/EditableMesh/EditableMesh.cpp b/Engine/Plugins/Runtime/EditableMesh/Source/EditableMesh/EditableMesh.cpp
--- a/Engine/Plugins/Runtime/EditableMesh/Source/EditableMesh/EditableMesh.cpp
+++ b/Engine/Plugins/Runtime/EditableMesh/Source/EditableMesh/EditableMesh.cpp
@@ -4,6 +4,8 @@
 #include "EditableMeshChanges.h"
 #include "EditableMeshCustomVersion.h"
 #include "EditableMeshAdapter.h"
+#include "EditableMeshOctree.h"
+#include "EngineDefines.h"	// For HALF_WORLD_MAX
 #include "HAL/IConsoleManager.h"
 #include "ProfilingDebugging/ScopedTimers.h"	// For FAutoScopedDurationTimer
 #include "GeomTools.h"
@@ -40,6 +42,9 @@ namespace EditableMesh
 {
 	static FAutoConsoleVariable InterpolatePositionsToLimit( TEXT( "EditableMesh.InterpolatePositionsToLimit" ), 1, TEXT( "Whether to interpolate vertex positions for subdivision meshes all the way to their limit surface position.  Otherwise, we stop at the most refined mesh position." ) );
 	static FAutoConsoleVariable InterpolateFVarsToLimit( TEXT( "EditableMesh.InterpolateFVarsToLimit" ), 1, TEXT( "Whether to interpolate face-varying vertex data for subdivision meshes all the way to their limit surface position.  Otherwise, we stop at the most refined mesh." ) );
+
+	static FAutoConsoleVariable OctreeIncrementalUpdateLimit( TEXT( "EditableMesh.OctreeIncrementalUpdateLimit" ), 0.4f, TEXT( "If more than this scalar percentage of polygons have changed, we'll rebuild the octree from scratch instead of incrementally updating it." ) );
+	static FAutoConsoleVariable UseBoundlessOctree( TEXT( "EditableMesh.UseBoundlessOctree" ), 1, TEXT( "If enabled, the octree for editable meshes will have a huge bounding box.  Otherwise, we'll compute a tightly wrapped bounds.  However, the bounds will not be able to grow beyond it's original size." ) );
 }
 
 
@@ -70,7 +75,8 @@ const FName UEditableMeshAttribute::PolygonCenterName( "PolygonCenter" );
 UEditableMesh::UEditableMesh()
 	: bAllowUndo( false ),
 	  bAllowCompact( false ),
-	  PendingCompactCounter( 0 )
+	  PendingCompactCounter( 0 ),
+	  bAllowSpatialDatabase( false )
 {
 }
 
@@ -92,11 +98,8 @@ void UEditableMesh::PostLoad()
 {
 	Super::PostLoad();
 
-	if( IsPreviewingSubdivisions() )
-	{
-		RefreshOpenSubdiv();
-	}
-
+	RefreshOpenSubdiv();
+	RebuildOctree();
 	RebuildRenderMesh();
 }
 
@@ -262,6 +265,39 @@ void UEditableMesh::FixUpElementIDs( const FElementIDRemappings& Remappings )
 			Polygon = Remappings.GetRemappedPolygonID( Polygon );
 		}
 	}
+
+	// Fix up spatial database
+	{
+		if( Octree.IsValid() )
+		{
+			for( FPolygonID& PolygonID : DeletedOctreePolygonIDs )
+			{
+				PolygonID = Remappings.GetRemappedPolygonID( PolygonID );
+			}
+
+			for( FPolygonID& PolygonID : NewOctreePolygonIDs )
+			{
+				PolygonID = Remappings.GetRemappedPolygonID( PolygonID );
+			}
+		}
+
+		{
+			// Make a temporary copy of the original map, and clear the stored version
+			static TMap<FPolygonID, FOctreeElementId> OldPolygonIDToOctreeElementIDMap;
+			OldPolygonIDToOctreeElementIDMap = PolygonIDToOctreeElementIDMap;
+			PolygonIDToOctreeElementIDMap.Reset();
+
+			// Rebuild the map with the remapped polygon IDs
+			for( const auto& OldPolygonIDAndOctreeElementID : OldPolygonIDToOctreeElementIDMap )
+			{
+				const FPolygonID OldPolygonID = OldPolygonIDAndOctreeElementID.Key;
+				const FOctreeElementId OctreeElementID = OldPolygonIDAndOctreeElementID.Value;
+
+				const FPolygonID NewPolygonID = Remappings.GetRemappedPolygonID( OldPolygonID );
+				PolygonIDToOctreeElementIDMap.Add( NewPolygonID, OctreeElementID );
+			}
+		}
+	}
 }
 
 
@@ -477,6 +513,12 @@ void UEditableMesh::EndModification( const bool bFromUndo )
 		// @todo mesheditor debug
 		// UE_LOG( LogEditableMesh, Log, TEXT( "UEditableStaticMesh::EndModification COMPLETE in %0.4fs" ), FunctionTimer.GetTime() );	  // @todo mesheditor: Shows bogus time values
 
+		// Rebuild our octree
+		if( CurrentModificationType == EMeshModificationType::Final )
+		{
+			UpdateOctreeIncrementallyIfPossibleOtherwiseRebuildIt();
+		}
+
 		FStartOrEndModificationChangeInput RevertInput;
 		RevertInput.bStartModification = true;
 		RevertInput.MeshModificationType = CurrentModificationType;
@@ -605,6 +647,15 @@ void UEditableMesh::SetVertexAttribute( const FVertexID VertexID, const FName At
 	{
 		check( AttributeIndex == 0 );	// Only one position is supported
 		Vertex.VertexPosition = AttributeValue;
+
+		// Update spatial database
+		if( Octree.IsValid() )
+		{
+			static TArray<FPolygonID> ConnectedPolygons;
+			GetVertexConnectedPolygons( VertexID, /* Out */ ConnectedPolygons );
+			DeletedOctreePolygonIDs.Append( ConnectedPolygons );
+			NewOctreePolygonIDs.Append( ConnectedPolygons );
+		}
 	}
 	else if( AttributeName == UEditableMeshAttribute::VertexCornerSharpness() )
 	{
@@ -4560,6 +4611,12 @@ void UEditableMesh::CreatePolygons( const TArray<FPolygonToCreate>& PolygonsToCr
 		Adapter->OnCreatePolygons( this, OutNewPolygonIDs );
 	}
 
+	// Update spatial database
+	if( Octree.IsValid() )
+	{
+		NewOctreePolygonIDs.Append( OutNewPolygonIDs );
+	}
+
 	// Generate tangent basis for the polygon
 	PolygonsPendingNewTangentBasis.Append( OutNewPolygonIDs );
 
@@ -4711,6 +4768,16 @@ void UEditableMesh::DeletePolygons( const TArray<FPolygonID>& PolygonIDsToDelete
 			Polygons.RemoveAt( PolygonID.GetValue() );
 		}
 
+		// Update spatial database
+		if( Octree.IsValid() )
+		{
+			DeletedOctreePolygonIDs.Append( PolygonIDsToDelete );
+			for( const FPolygonID PolygonID : PolygonIDsToDelete )
+			{
+				NewOctreePolygonIDs.Remove( PolygonID );
+			}
+		}
+
 		// Remove vertex instances which are exclusively used by this polygon.
 		// We do not want this to remove orphaned vertices; this will optionally happen below when removing edges.
 		if( OrphanedVertexInstanceIDs.Num() > 0 )
@@ -9049,6 +9116,233 @@ void UEditableMesh::QuadrangulatePolygonGroup( const FPolygonGroupID PolygonGrou
 }
 
 
+void UEditableMesh::UpdateOctreeIncrementallyIfPossibleOtherwiseRebuildIt()
+{
+	bool bAnythingChanged = true;
+	bool bDoIncrementalUpdate = false;
+
+	if( Octree.IsValid() )
+	{
+		if( IsSpatialDatabaseAllowed() )
+		{
+			if( NewOctreePolygonIDs.Num() == 0 &&
+				DeletedOctreePolygonIDs.Num() == 0 )
+			{
+				// Nothing has changed!
+				bAnythingChanged = false;
+			}
+			else
+			{
+				// Make sure we haven't changed or deleted so many polygons that it's not even worth doing an incremental
+				// update to the mesh.  It's generally more expensive to remove things from the octree than to add to it,
+				// because it will need to consider how to collapse nodes.  If we're only adding new things, then we'll
+				// never rebuild it from scratch.
+				const float ScalarPercentOfPolygonsChangedOrDeleted = (float)DeletedOctreePolygonIDs.Num() / (float)GetPolygonCount();
+				if( ScalarPercentOfPolygonsChangedOrDeleted < EditableMesh::OctreeIncrementalUpdateLimit->GetFloat() )
+				{
+					// We have a reasonable number of incremental changes, so let's go ahead and make those!
+					bDoIncrementalUpdate = true;
+				}
+			}
+		}
+	}
+
+	if( bAnythingChanged )
+	{
+		if( bDoIncrementalUpdate )
+		{
+			check( Octree.IsValid() );
+
+			// Clear out deleted polygons from our octree
+			{
+				for( const FPolygonID PolygonID : DeletedOctreePolygonIDs )
+				{
+					const FOctreeElementId OctreeElementID = PolygonIDToOctreeElementIDMap.FindAndRemoveChecked( PolygonID );
+					Octree->RemoveElement( OctreeElementID );
+				}
+				DeletedOctreePolygonIDs.Reset();
+			}
+
+			// Now, add new polygons to the octree
+			{
+				for( const FPolygonID PolygonID : NewOctreePolygonIDs )
+				{
+					checkSlow( IsValidPolygon( PolygonID ) );
+
+					FBox BoundingBox;
+					{
+						BoundingBox.Init();
+
+						static TArray<FVertexID> PerimeterVertices;
+						GetPolygonPerimeterVertices( PolygonID, /* Out */ PerimeterVertices );
+						for( const FVertexID VertexID : PerimeterVertices )
+						{
+							const FVector VertexPosition = GetVertexAttribute( VertexID, UEditableMeshAttribute::VertexPosition(), 0 );
+							BoundingBox += VertexPosition;
+						}
+					}
+
+					const FEditableMeshOctreePolygon OctreePolygon( *this, PolygonID, FBoxCenterAndExtent( BoundingBox ) );
+					Octree->AddElement( OctreePolygon );
+				}
+				NewOctreePolygonIDs.Reset();
+			}
+		}
+		else
+		{
+			RebuildOctree();
+		}
+	}
+}
+
+
+void UEditableMesh::RebuildOctree()
+{
+	Octree.Reset();
+	PolygonIDToOctreeElementIDMap.Reset();
+	NewOctreePolygonIDs.Reset();
+	DeletedOctreePolygonIDs.Reset();
+
+	if( IsSpatialDatabaseAllowed() )
+	{
+		FVector OctreeOrigin;
+		float OctreeExtent;
+		{
+			if( EditableMesh::UseBoundlessOctree->GetInt() != 0 )
+			{
+				// We use a 'boundless' octree for our mesh.  This is because we don't want to have to update the root node bounds
+				// of the tree as we change the mesh, and also to avoid having to compute a bounding box here.
+				OctreeOrigin = FVector::ZeroVector;
+				OctreeExtent = HALF_WORLD_MAX;
+			}
+			else
+			{
+				// Compute a bounding box to use that tightly wraps the mesh.
+
+				// NOTE: This will have problems.  The root bounding box cannot change without rebuilding the octree, so if the
+				// mesh is modified and grown, elements outside the original bounds will not be selectable!  Additionally, we
+				// incur the additional performance cost of computing a fresh bounding box here.
+
+				const FBox MeshBoundingBox = this->ComputeBoundingBox();
+				OctreeOrigin = MeshBoundingBox.GetCenter();
+				OctreeExtent = MeshBoundingBox.GetExtent().GetAbsMax();
+			}
+		}
+
+		Octree = MakeShareable( new FEditableMeshOctree( OctreeOrigin, OctreeExtent ) );
+
+		// @todo mesheditor spatial: We'll currently only ever find vertices and edges that are connected to polygons.  Our spatial
+		// database only tracks polygons.  Do we need to support hover/selection of "loose" vertices and edges also?
+
+		const uint32 PolygonArraySize = GetPolygonArraySize();
+		for( uint32 PolygonIndex = 0; PolygonIndex < PolygonArraySize; ++PolygonIndex )
+		{
+			const FPolygonID PolygonID( PolygonIndex );
+			if( IsValidPolygon( PolygonID ) )
+			{
+				FBox BoundingBox;
+				{
+					BoundingBox.Init();
+
+					static TArray<FVertexID> PerimeterVertices;
+					GetPolygonPerimeterVertices( PolygonID, /* Out */ PerimeterVertices );
+					for( const FVertexID VertexID : PerimeterVertices )
+					{
+						const FVector VertexPosition = GetVertexAttribute( VertexID, UEditableMeshAttribute::VertexPosition(), 0 );
+						BoundingBox += VertexPosition;
+					}
+				}
+
+				const FEditableMeshOctreePolygon OctreePolygon( *this, PolygonID, FBoxCenterAndExtent( BoundingBox ) );
+				Octree->AddElement( OctreePolygon );
+			}
+		}
+
+		// @todo mesheditor debug
+		// Octree->DumpStats();
+	}
+}
+
+
+void UEditableMesh::SearchSpatialDatabaseForPolygonsPotentiallyIntersectingLineSegment( const FVector LineSegmentStart, const FVector LineSegmentEnd, TArray<FPolygonID>& OutPolygons ) const
+{
+	OutPolygons.Reset();
+
+	// @todo mesheditor perf: Ideally we "early out" of octree traversal, by walking through the octree along the ray, reporting back to QueryElement to find out when we should stop
+	// @todo mesheditor scripting: Should spit a warning for Blueprint users if Octree is not allowed when calling this function
+
+	if( IsSpatialDatabaseAllowed() && ensure( Octree.IsValid() ) )
+	{
+		const FVector LineSegmentVector = LineSegmentEnd - LineSegmentStart;
+		const FVector LineSegmentVectorReciprocal = LineSegmentVector.Reciprocal();
+
+		// @todo mesheditor perf: Do we need to use a custom stack allocator for iterating?  The default should probably be okay.
+		for( FEditableMeshOctree::TConstIterator<> OctreeIt( *Octree );
+			 OctreeIt.HasPendingNodes();
+			 OctreeIt.Advance() )
+		{
+			const FEditableMeshOctree::FNode& OctreeNode = OctreeIt.GetCurrentNode();
+			const FOctreeNodeContext& OctreeNodeContext = OctreeIt.GetCurrentContext();
+
+			// Leaf nodes have no children, so don't bother iterating
+			if( !OctreeNode.IsLeaf() )
+			{
+				// Find children of this octree node that overlap our line segment
+				FOREACH_OCTREE_CHILD_NODE( ChildRef )
+				{
+					if( OctreeNode.HasChild( ChildRef ) )
+					{
+						const FOctreeNodeContext ChildContext = OctreeNodeContext.GetChildContext( ChildRef );
+
+						// @todo mesheditor: LineBoxIntersection() has a magic number in its implementation we might want to look at (search for BOX_SIDE_THRESHOLD)
+						const bool bIsOverlappingLineSegment =
+							FMath::LineBoxIntersection(
+								ChildContext.Bounds.GetBox(),
+								LineSegmentStart,
+								LineSegmentEnd,
+								LineSegmentVector,
+								LineSegmentVectorReciprocal );
+
+						if( bIsOverlappingLineSegment )
+						{
+							// DrawDebugBox( GWorld, ChildContext.Bounds.Center, ChildContext.Bounds.Extent * 0.8f, FQuat::Identity, FColor::Green, false, 0.0f );		// @todo mesheditor debug: (also, wrong coordinate system!)
+
+							// Push it on the iterator's pending node stack.
+							OctreeIt.PushChild( ChildRef );
+						}
+						else
+						{
+							// DrawDebugBox( GWorld, ChildContext.Bounds.Center, ChildContext.Bounds.Extent, FQuat::Identity, FColor( 128, 128, 128 ), false, 0.0f );	// @todo mesheditor debug: (also, wrong coordinate system!)
+						}
+					}
+				}
+			}
+
+			// All of the elements in this octree node are candidates.  Note this node may not be a leaf node, and that's OK.
+			for( FEditableMeshOctree::ElementConstIt OctreeElementIt( OctreeNode.GetElementIt() ); OctreeElementIt; ++OctreeElementIt )
+			{
+				const FEditableMeshOctreePolygon& OctreePolygon = *OctreeElementIt;
+				OutPolygons.Add( OctreePolygon.PolygonID );
+			}
+		}
+	}
+}
+
+
+void UEditableMesh::SetAllowSpatialDatabase( const bool bInAllowSpatialDatabase )
+{
+	if( bAllowSpatialDatabase != bInAllowSpatialDatabase )
+	{
+		bAllowSpatialDatabase = bInAllowSpatialDatabase;
+
+		if( !IsBeingModified() )
+		{
+			RebuildOctree();
+		}
+	}
+}
+
+
 bool UEditableMesh::AnyChangesToUndo() const
 {
 	return bAllowUndo && Undo.IsValid() && Undo->Subchanges.Num() > 0;
diff --git a/Engine/Plugins/Runtime/EditableMesh/Source/EditableMesh/EditableMeshOctree.h b/Engine/Plugins/Runtime/EditableMesh/Source/EditableMesh/EditableMeshOctree.h
new file mode 100644
--- /dev/null
+++ b/Engine/Plugins/Runtime/EditableMesh/Source/EditableMesh/EditableMeshOctree.h
@@ -0,0 +1,66 @@
+// Copyright 1998-2015 Epic Games, Inc. All Rights Reserved.
+
+#pragma once
+
+#include "GenericOctreePublic.h"
+#include "GenericOctree.h"
+#include "EditableMesh.h"
+
+
+struct FEditableMeshOctreePolygon
+{
+	class UEditableMesh& EditableMesh;
+	FPolygonID PolygonID;
+	FBoxCenterAndExtent PolygonBounds;
+
+	FEditableMeshOctreePolygon( UEditableMesh& InitEditableMesh, const FPolygonID& InitPolygonID, const FBoxCenterAndExtent& InitPolygonBounds )
+		: EditableMesh( InitEditableMesh ),
+		  PolygonID( InitPolygonID ),
+		  PolygonBounds( InitPolygonBounds )
+	{
+	}
+};
+
+
+struct FEditableMeshOctreeSemantics
+{
+	// @todo mesheditor perf: Profile and tweak these octree settings
+
+	// When a leaf gets more than this number of elements, it will split itself into a node with multiple child leaves
+	enum { MaxElementsPerLeaf = 6 };
+
+	// This is used for incremental updates.  When removing a polygon, larger values will cause leaves to be removed and collapsed into a parent node.
+	enum { MinInclusiveElementsPerNode = 7 };
+
+	// How deep the tree can go.
+	enum { MaxNodeDepth = 20 };
+
+
+	typedef TInlineAllocator<MaxElementsPerLeaf> ElementAllocator;
+
+	FORCEINLINE static FBoxCenterAndExtent GetBoundingBox( const FEditableMeshOctreePolygon& Element )
+	{
+		return Element.PolygonBounds;
+	}
+
+	FORCEINLINE static bool AreElementsEqual( const FEditableMeshOctreePolygon& A, const FEditableMeshOctreePolygon& B )
+	{
+		return ( A.PolygonID == B.PolygonID );
+	}
+
+	FORCEINLINE static void SetElementId( const FEditableMeshOctreePolygon& Element, FOctreeElementId OctreeElementID )
+	{
+		Element.EditableMesh.PolygonIDToOctreeElementIDMap.Add( Element.PolygonID, OctreeElementID );
+	}
+};
+
+class FEditableMeshOctree : public TOctree<FEditableMeshOctreePolygon, FEditableMeshOctreeSemantics>
+{
+public:
+	FEditableMeshOctree( const FVector& InOrigin, float InExtent )
+		: TOctree( InOrigin, InExtent )
+	{
+	}
+};
+
+
diff --git a/Engine/Plugins/Runtime/EditableMesh/Source/EditableMesh/Public/EditableMesh.h b/Engine/Plugins/Runtime/EditableMesh/Source/EditableMesh/Public/EditableMesh.h
--- a/Engine/Plugins/Runtime/EditableMesh/Source/EditableMesh/Public/EditableMesh.h
+++ b/Engine/Plugins/Runtime/EditableMesh/Source/EditableMesh/Public/EditableMesh.h
@@ -9,6 +9,7 @@
 #include "Change.h"		// For TUniquePtr<FChange>
 #include "Logging/LogMacros.h"
 #include "Materials/MaterialInterface.h"
+#include "GenericOctreePublic.h"	// For FOctreeElementId
 #include "EditableMesh.generated.h"
 
 class UEditableMeshAdapter;
@@ -601,6 +602,24 @@ public:
 		bAllowUndo = bInAllowUndo;
 	}
 
+	/**
+	 * @return	Returns true if our octree spatial database is enabled for this mesh
+	 */
+	UFUNCTION( BlueprintPure, Category="Editable Mesh" )
+	bool IsSpatialDatabaseAllowed() const
+	{
+		return bAllowSpatialDatabase;
+	}
+
+	/**
+	 * Sets whether this mesh should automatically generate and maintain an octree spatial database.  Certain queries may only be
+	 * supported when the mesh has an octree generated.  The octree is never saved or loaded, and always generated on demand.  This
+	 * feature adds significant overhead to editable mesh initialization and modification, so only use it if you really need to.
+	 *
+	 * @param	bInAllowSpatialDatabase		True if an octree should be generated and maintained for this mesh.
+	 */
+	UFUNCTION( BlueprintCallable, Category="Editable Mesh" )
+	void SetAllowSpatialDatabase( const bool bInAllowSpatialDatabase );
 
 	/**
 	 * @return	Returns true if there are any existing tracked changes that can be undo right now
@@ -731,10 +750,9 @@ public:
 	UFUNCTION( BlueprintPure, Category="Editable Mesh" ) void ComputeTextureCoordinatesForPointOnPolygon( const FPolygonID PolygonID, const FVector PointOnPolygon, bool& bOutWereTextureCoordinatesFound, TArray<FVector4>& OutInterpolatedTextureCoordinates );
 	UFUNCTION( BlueprintPure, Category="Editable Mesh" ) void ComputePolygonsSharedEdges( const TArray<FPolygonID>& PolygonIDs, TArray<FEdgeID>& OutSharedEdgeIDs ) const;
 	UFUNCTION( BlueprintPure, Category="Editable Mesh" ) void FindPolygonLoop( const FEdgeID EdgeID, TArray<FEdgeID>& OutEdgeLoopEdgeIDs, TArray<FEdgeID>& OutFlippedEdgeIDs, TArray<FEdgeID>& OutReversedEdgeIDPathToTake, TArray<FPolygonID>& OutPolygonIDsToSplit ) const;
+	UFUNCTION( BlueprintPure, Category="Editable Mesh" ) void SearchSpatialDatabaseForPolygonsPotentiallyIntersectingLineSegment( const FVector LineSegmentStart, const FVector LineSegmentEnd, TArray<FPolygonID>& OutPolygons ) const;
 
 
-	// Low poly editing features
-	// @todo mesheditor: Move elsewhere
 	UFUNCTION( BlueprintCallable, Category="Editable Mesh" ) void SetSubdivisionCount( const int32 NewSubdivisionCount );
 	UFUNCTION( BlueprintCallable, Category="Editable Mesh" ) void MoveVertices( const TArray<FVertexToMove>& VerticesToMove );
 	UFUNCTION( BlueprintCallable, Category="Editable Mesh" ) void CreateMissingPolygonPerimeterEdges( const FPolygonID PolygonID, TArray<FEdgeID>& OutNewEdgeIDs );
@@ -814,6 +832,16 @@ private:
 	/** Called during end modification to retriangulate polygons in the pending polygon list */
 	void RetriangulatePolygons();
 
+	/** Tries to incrementally update the octree based off geometry that has changed since last time.  If that's not
+	    reasonable to do, then this will rebuild the entire octree from scratch */
+	void UpdateOctreeIncrementallyIfPossibleOtherwiseRebuildIt();
+
+public:
+	// @todo mesheditor: temporarily changed access to public so the adapter can call it when building the editable mesh from the source static mesh. Think about this some more.
+	/** Rebuilds the octree */
+	void RebuildOctree();
+
+
 public:
 
 	/** Each editable vertex in this mesh. */
@@ -907,4 +935,27 @@ public:
 
 	/** Broadcast event when element IDs are remapped (for example, following Compact / Uncompact) */
 	FElementIDsRemapped ElementIDsRemappedEvent;
+
+
+	//
+	// Spatial database
+	//
+
+	/** True if we should generate and maintain an octree spatial database for this mesh. */
+	bool bAllowSpatialDatabase;
+
+	/** Octree to accelerate spatial queries against the mesh.  This is never serialized, and only created and maintained if bAllowSpatialDatabase is enabled. */
+	TSharedPtr<class FEditableMeshOctree> Octree;
+
+	/** Maps our polygon IDs to octree element IDs */
+	TMap<FPolygonID, FOctreeElementId> PolygonIDToOctreeElementIDMap;
+
+	/** Polygons that were deleted since the last time our octree was refreshed */
+	TSet<FPolygonID> DeletedOctreePolygonIDs;
+
+	/** Newly-created polygons since the last time our octree was refreshed */
+	TSet<FPolygonID> NewOctreePolygonIDs;
+
+	friend struct FEditableMeshOctreeSemantics;	// NOTE: Allows inline access to PolygonIDToOctreeElementIDMap
+
 };
