Mesh Editor: UI improvements
- Renamed lots of UI text and tooltips to improve UX
- Hid the 'Propagate' and 'Per Instance' buttons until this feature is properly supported
- Flipped radio buttons so they're on top of 'instant' buttons
- Renamed 'Any' to 'Mesh' in the top level UI (meaning, 'any part of the mesh')
- Added group name title text for 'Mesh', 'Polygon', 'Edge', 'Vertex'
- Reversed order of top level buttons, removed 'Select' label
- 'Remove Subdivision Level' now greys out if there is nothing already subdivided
- Fixed lots of non-unity build issues

diff --git a/Engine/Plugins/Editor/MeshEditor/Source/MeshEditor/MeshEditorCommands.cpp b/Engine/Plugins/Editor/MeshEditor/Source/MeshEditor/MeshEditorCommands.cpp
--- a/Engine/Plugins/Editor/MeshEditor/Source/MeshEditor/MeshEditorCommands.cpp
+++ b/Engine/Plugins/Editor/MeshEditor/Source/MeshEditor/MeshEditorCommands.cpp
@@ -140,9 +140,9 @@ FMeshEditorVertexCommands::FMeshEditorVertexCommands()
 
 void FMeshEditorVertexCommands::RegisterCommands()
 {
-	UI_COMMAND(MoveVertex, "Move Vertex Mode", "Set the primary action to move vertices.", EUserInterfaceActionType::RadioButton, FInputChord());
+	UI_COMMAND(MoveVertex, "Move", "Move selected vertices using a transform gizmo, or click and drag to move vertices directly.", EUserInterfaceActionType::RadioButton, FInputChord());
 
-	UI_COMMAND(WeldVertices, "Weld Vertices", "Weld the selected vertices, keeping the first selected vertex.", EUserInterfaceActionType::Button, FInputChord());
+	UI_COMMAND(WeldVertices, "Weld", "Weld the selected vertices, keeping the first selected vertex.", EUserInterfaceActionType::Button, FInputChord());
 
 	for( UMeshEditorCommand* Command : MeshEditorCommands::Get() )
 	{
@@ -165,7 +165,7 @@ FMeshEditorEdgeCommands::FMeshEditorEdgeCommands()
 
 void FMeshEditorEdgeCommands::RegisterCommands()
 {
-	UI_COMMAND(MoveEdge, "Move Edge Mode", "Set the primary action to move edges.", EUserInterfaceActionType::RadioButton, FInputChord());
+	UI_COMMAND(MoveEdge, "Move", "Move selected edges using a transform gizmo, or click and drag to move edges directly.", EUserInterfaceActionType::RadioButton, FInputChord());
 
 	UI_COMMAND(SelectEdgeLoop, "Select Edge Loop", "Select the edge loops which contain the selected edges.", EUserInterfaceActionType::Button, FInputChord(EKeys::Two, EModifierKey::Shift));
 
@@ -190,9 +190,9 @@ FMeshEditorPolygonCommands::FMeshEditorPolygonCommands()
 
 void FMeshEditorPolygonCommands::RegisterCommands()
 {
-	UI_COMMAND(MovePolygon, "Move Polygon Mode", "Set the primary action to move polygons.", EUserInterfaceActionType::RadioButton, FInputChord());
-	UI_COMMAND(FlipPolygon, "Flip Polygon", "Flip the currently selected polygons.", EUserInterfaceActionType::Button, FInputChord(EKeys::F, EModifierKey::Shift));
-	UI_COMMAND(TriangulatePolygon, "Triangulate Polygon", "Triangulate the currently selected polygons.", EUserInterfaceActionType::Button, FInputChord(EKeys::T));
+	UI_COMMAND(MovePolygon, "Move", "Move selected polygons using a transform gizmo, or click and drag to move polygons directly.", EUserInterfaceActionType::RadioButton, FInputChord());
+	UI_COMMAND(FlipPolygon, "Flip", "Flip the currently selected polygons.", EUserInterfaceActionType::Button, FInputChord(EKeys::F, EModifierKey::Shift));
+	UI_COMMAND(TriangulatePolygon, "Triangulate", "Triangulate the currently selected polygons.", EUserInterfaceActionType::Button, FInputChord(EKeys::T));
 	UI_COMMAND(AssignMaterial, "Assign Material", "Assigns the highlighted material in the Content Browser to the currently selected polygons.", EUserInterfaceActionType::Button, FInputChord(EKeys::M));
 
 	for( UMeshEditorCommand* Command : MeshEditorCommands::Get() )
diff --git a/Engine/Plugins/Editor/MeshEditor/Source/MeshEditor/MeshEditorMode.cpp b/Engine/Plugins/Editor/MeshEditor/Source/MeshEditor/MeshEditorMode.cpp
--- a/Engine/Plugins/Editor/MeshEditor/Source/MeshEditor/MeshEditorMode.cpp
+++ b/Engine/Plugins/Editor/MeshEditor/Source/MeshEditor/MeshEditorMode.cpp
@@ -622,7 +622,24 @@ void FMeshEditorMode::BindCommands()
 
 	// Register commands which work even without a selected element, as long as at least one mesh is selected
 	RegisterCommonCommand( MeshEditorCommonCommands.AddSubdivisionLevel, FExecuteAction::CreateLambda( [this] { AddOrRemoveSubdivisionLevel( true ); } ) );
-	RegisterCommonCommand( MeshEditorCommonCommands.RemoveSubdivisionLevel, FExecuteAction::CreateLambda( [this] { AddOrRemoveSubdivisionLevel( false ); } ) );
+	RegisterCommonCommand( MeshEditorCommonCommands.RemoveSubdivisionLevel, 
+		FExecuteAction::CreateLambda( [this] { AddOrRemoveSubdivisionLevel( false ); } ),
+		FCanExecuteAction::CreateLambda( [this] 
+			{ 
+				// Only allow 'Remove' if any selected meshes are already subdivided
+				bool bAnySubdividedMeshes = false;
+				for( const UEditableMesh* EditableMesh : GetSelectedEditableMeshes() )
+				{
+					if( EditableMesh->GetSubdivisionCount() > 0 )
+					{
+						bAnySubdividedMeshes = true;
+						break;
+					}
+				}
+				return bAnySubdividedMeshes;
+			}
+		) 
+	);
 
 	// @todo mesheditor: support EUserInterfaceActionType::ToggleButton actions in the UI, and extend RegisterCommand to allow
 	// a delegate returning check state.
@@ -656,9 +673,6 @@ void FMeshEditorMode::BindCommands()
 					// Common action
 					FUIAction UIAction = Command->MakeUIAction( *this );
 					CommonActions.Emplace( Command->GetUICommandInfo(), UIAction );
-					VertexActions.Emplace( Command->GetUICommandInfo(), UIAction );
-					EdgeActions.Emplace( Command->GetUICommandInfo(), UIAction );
-					PolygonActions.Emplace( Command->GetUICommandInfo(), UIAction );
 				}
 				break;
 
@@ -748,13 +762,10 @@ void FMeshEditorMode::RegisterPolygonEditingMode( const TSharedPtr<FUICommandInf
 }
 
 
-void FMeshEditorMode::RegisterCommonCommand( const TSharedPtr<FUICommandInfo>& Command, const FExecuteAction& ExecuteAction )
+void FMeshEditorMode::RegisterCommonCommand( const TSharedPtr<FUICommandInfo>& Command, const FExecuteAction& ExecuteAction, const FCanExecuteAction CanExecuteAction )
 {
-	FCanExecuteAction CanExecute = FCanExecuteAction::CreateLambda( [this] { return GetSelectedEditableMeshes().Num() > 0; } );
-	CommonActions.Emplace( Command, FUIAction( ExecuteAction, CanExecute ) );
-	VertexActions.Emplace( Command, FUIAction( ExecuteAction, CanExecute ) );
-	EdgeActions.Emplace( Command, FUIAction( ExecuteAction, CanExecute ) );
-	PolygonActions.Emplace( Command, FUIAction( ExecuteAction, CanExecute ) );
+	FCanExecuteAction CompositeCanExecuteAction = FCanExecuteAction::CreateLambda( [this, CanExecuteAction] { return GetSelectedEditableMeshes().Num() > 0 && ( !CanExecuteAction.IsBound() || CanExecuteAction.Execute() ); } );
+	CommonActions.Emplace( Command, FUIAction( ExecuteAction, CompositeCanExecuteAction ) );
 }
 
 
diff --git a/Engine/Plugins/Editor/MeshEditor/Source/MeshEditor/MeshEditorMode.h b/Engine/Plugins/Editor/MeshEditor/Source/MeshEditor/MeshEditorMode.h
--- a/Engine/Plugins/Editor/MeshEditor/Source/MeshEditor/MeshEditorMode.h
+++ b/Engine/Plugins/Editor/MeshEditor/Source/MeshEditor/MeshEditorMode.h
@@ -281,7 +281,7 @@ protected:
 	void RegisterEdgeEditingMode( const TSharedPtr<FUICommandInfo>& Command, FName EditingMode );
 	void RegisterPolygonEditingMode( const TSharedPtr<FUICommandInfo>& Command, FName EditingMode );
 
-	void RegisterCommonCommand( const TSharedPtr<FUICommandInfo>& Command, const FExecuteAction& ExecuteAction );
+	void RegisterCommonCommand( const TSharedPtr<FUICommandInfo>& Command, const FExecuteAction& ExecuteAction, const FCanExecuteAction CanExecuteAction = FCanExecuteAction() );
 	void RegisterAnyElementCommand( const TSharedPtr<FUICommandInfo>& Command, const FExecuteAction& ExecuteAction );
 	void RegisterVertexCommand( const TSharedPtr<FUICommandInfo>& Command, const FExecuteAction& ExecuteAction );
 	void RegisterEdgeCommand( const TSharedPtr<FUICommandInfo>& Command, const FExecuteAction& ExecuteAction );
diff --git a/Engine/Plugins/Editor/MeshEditor/Source/MeshEditor/MeshEditorModeToolkit.cpp b/Engine/Plugins/Editor/MeshEditor/Source/MeshEditor/MeshEditorModeToolkit.cpp
--- a/Engine/Plugins/Editor/MeshEditor/Source/MeshEditor/MeshEditorModeToolkit.cpp
+++ b/Engine/Plugins/Editor/MeshEditor/Source/MeshEditor/MeshEditorModeToolkit.cpp
@@ -5,8 +5,15 @@
 #include "MeshEditorCommands.h"
 #include "MeshEditorStyle.h"
 #include "SWidgetSwitcher.h"
-#include "SSeparator.h"
+#include "SButton.h"
 #include "SScrollBox.h"
+#include "SToolTip.h"
+#include "SCheckBox.h"
+#include "SSpacer.h"
+#include "STextBlock.h"
+#include "SSeparator.h"
+#include "SBox.h"
+#include "EditorStyleSet.h"
 
 #define LOCTEXT_NAMESPACE "MeshEditorModeToolkit"
 
@@ -20,7 +27,7 @@ public:
 public:
 
 	/** SCompoundWidget functions */
-	void Construct( const FArguments& InArgs, const TArray<TTuple<TSharedPtr<FUICommandInfo>, FUIAction>>& Actions )
+	void Construct( const FArguments& InArgs, const FText& GroupName, const TArray<TTuple<TSharedPtr<FUICommandInfo>, FUIAction>>& Actions )
 	{
 		TSharedRef<SVerticalBox> Buttons = SNew( SVerticalBox );
 		TSharedRef<SVerticalBox> RadioButtons = SNew( SVerticalBox );
@@ -32,9 +39,10 @@ public:
 
 			if( CommandInfo.GetUserInterfaceType() == EUserInterfaceActionType::Button )
 			{
+				const bool bIsFirstItem = Buttons->NumSlots() == 0;
 				Buttons->AddSlot()
 				.AutoHeight()
-				.Padding( 3.0f )
+				.Padding( 3.0f, bIsFirstItem ? 9.0f : 3.0f, 3.0f, 3.0f )
 				[
 					SNew( SButton )
 					.HAlign( HAlign_Center )
@@ -48,9 +56,10 @@ public:
 			}
 			else if( CommandInfo.GetUserInterfaceType() == EUserInterfaceActionType::RadioButton )
 			{
+				const bool bIsFirstItem = RadioButtons->NumSlots() == 0;
 				RadioButtons->AddSlot()
 				.AutoHeight()
-				.Padding( 1.0f )
+				.Padding( 3.0f, bIsFirstItem ? 7.0f : 1.0f, 3.0f, 1.0f )
 				[
 					SNew( SCheckBox )
 					.Style( FMeshEditorStyle::Get(), "EditingMode.Entry" )
@@ -88,24 +97,24 @@ public:
 			SNew( SVerticalBox )
 			+SVerticalBox::Slot()
 			.AutoHeight()
-			.Padding( 6.0f )
+			.Padding( 6.0f, 6.0f, 6.0f, 2.0f )
 			.HAlign( HAlign_Center )
 			[
-				Buttons
+				SNew( STextBlock )
+				.TextStyle( FMeshEditorStyle::Get(), "EditingMode.GroupName.Text" )
+				.Text( GroupName )
 			]
 			+SVerticalBox::Slot()
 			.AutoHeight()
-			.Padding( 6.0f )
+			.HAlign( HAlign_Center )
 			[
-				SNew( SSeparator )
-				.Orientation( Orient_Horizontal )
+				RadioButtons
 			]
 			+SVerticalBox::Slot()
 			.AutoHeight()
-			.Padding( 6.0f )
 			.HAlign( HAlign_Center )
 			[
-				RadioButtons
+				Buttons
 			]
 		];
 	}
@@ -168,40 +177,27 @@ void SMeshEditorModeControls::Construct( const FArguments& InArgs, IMeshEditorMo
 
 	WidgetSwitcher->AddSlot( static_cast<int32>( EEditableMeshElementType::Vertex ) )
 	[
-		SNew( SMeshEditorModeControlWidget, MeshEditorMode.GetVertexActions() )
+		SNew( SMeshEditorModeControlWidget, LOCTEXT( "VertexGroupName", "Vertex" ), MeshEditorMode.GetVertexActions() )
 	];
 
 	WidgetSwitcher->AddSlot( static_cast<int32>( EEditableMeshElementType::Edge ) )
 	[
-		SNew( SMeshEditorModeControlWidget, MeshEditorMode.GetEdgeActions() )
+		SNew( SMeshEditorModeControlWidget, LOCTEXT( "EdgeGroupName", "Edge" ), MeshEditorMode.GetEdgeActions() )
 	];
 
 	WidgetSwitcher->AddSlot( static_cast<int32>( EEditableMeshElementType::Polygon ) )
 	[
-		SNew( SMeshEditorModeControlWidget, MeshEditorMode.GetPolygonActions() )
+		SNew( SMeshEditorModeControlWidget, LOCTEXT( "PolygonGroupName", "Polygon" ), MeshEditorMode.GetPolygonActions() )
 	];
 
 	WidgetSwitcher->AddSlot( static_cast<int32>( EEditableMeshElementType::Invalid ) )
 	[
-		SNew( SVerticalBox )
-		+ SVerticalBox::Slot()
+		SNew( SBox )
+		.Padding( 20.0f )
+		.HAlign( HAlign_Center )
 		[
-			SNew( SBox )
-			.Visibility_Lambda( [&MeshEditorMode]() { return MeshEditorMode.GetSelectedEditableMeshes().Num() > 0 ? EVisibility::Visible : EVisibility::Collapsed; } )
-			[
-				SNew( SMeshEditorModeControlWidget, MeshEditorMode.GetCommonActions() )
-			]
-		]
-		+SVerticalBox::Slot()
-		[
-			SNew( SBox )
-			.Padding( 20.0f )
-			.HAlign( HAlign_Center )
-			.Visibility_Lambda( [&MeshEditorMode]() { return MeshEditorMode.GetSelectedEditableMeshes().Num() == 0 ? EVisibility::Visible : EVisibility::Collapsed; } )
-			[
-				SNew( STextBlock )
-				.Text( LOCTEXT( "NothingSelected", "Nothing is selected" ) )
-			]
+			SNew( STextBlock )
+			.Text( LOCTEXT( "NothingSelected", "Please select a mesh to edit." ) )
 		]
 	];
 
@@ -217,36 +213,28 @@ void SMeshEditorModeControls::Construct( const FArguments& InArgs, IMeshEditorMo
 			[
 				SNew( SHorizontalBox )
 				+SHorizontalBox::Slot()
-				.VAlign( VAlign_Center )
-				.AutoWidth()
-				.Padding( 2, 2, 8, 2 )
-				[
-					SNew( STextBlock )
-					.Text( LOCTEXT( "SelectLabel", "Select" ) )
-				]
-				+SHorizontalBox::Slot()
 				.FillWidth( 1 )
 				.Padding( 2 )
 				[
-					SNew( SMeshEditorSelectionModeWidget, MeshEditorMode, EEditableMeshElementType::Vertex, LOCTEXT( "Vertices", "Vertices" ) )
+					SNew( SMeshEditorSelectionModeWidget, MeshEditorMode, EEditableMeshElementType::Any, LOCTEXT( "AnyElementType", "Mesh" ) )
 				]
 				+SHorizontalBox::Slot()
 				.FillWidth( 1 )
 				.Padding( 2 )
 				[
-					SNew( SMeshEditorSelectionModeWidget, MeshEditorMode, EEditableMeshElementType::Edge, LOCTEXT( "Edges", "Edges" ) )
+					SNew( SMeshEditorSelectionModeWidget, MeshEditorMode, EEditableMeshElementType::Polygon, LOCTEXT( "Polygon", "Polygon" ) )
 				]
 				+SHorizontalBox::Slot()
 				.FillWidth( 1 )
 				.Padding( 2 )
 				[
-					SNew( SMeshEditorSelectionModeWidget, MeshEditorMode, EEditableMeshElementType::Polygon, LOCTEXT( "Polygons", "Polygons" ) )
+					SNew( SMeshEditorSelectionModeWidget, MeshEditorMode, EEditableMeshElementType::Edge, LOCTEXT( "Edge", "Edge" ) )
 				]
 				+SHorizontalBox::Slot()
 				.FillWidth( 1 )
 				.Padding( 2 )
 				[
-					SNew( SMeshEditorSelectionModeWidget, MeshEditorMode, EEditableMeshElementType::Any, LOCTEXT( "Any", "Any" ) )
+					SNew( SMeshEditorSelectionModeWidget, MeshEditorMode, EEditableMeshElementType::Vertex, LOCTEXT( "Vertex", "Vertex" ) )
 				]
 			]
 			+SVerticalBox::Slot()
@@ -255,6 +243,7 @@ void SMeshEditorModeControls::Construct( const FArguments& InArgs, IMeshEditorMo
 			.HAlign( HAlign_Right )
 			[
 				SNew( SHorizontalBox )
+				.Visibility( EVisibility::Collapsed )	// @todo mesheditor instancing: UI for instancing features is disabled until this feature is working properly
 				+SHorizontalBox::Slot()
 				.AutoWidth()
 				.Padding( 2 )
@@ -296,11 +285,51 @@ void SMeshEditorModeControls::Construct( const FArguments& InArgs, IMeshEditorMo
 				]
 			]
 			+SVerticalBox::Slot()
+			.AutoHeight()
 			[
 				SNew( SBorder )
 				.BorderImage( FEditorStyle::GetBrush( "ToolPanel.GroupBorder" ) )
 				[
-					WidgetSwitcher
+					SNew( SVerticalBox )
+					+ SVerticalBox::Slot()
+					.AutoHeight()
+					[
+						SNew( SBox )
+						.Padding( 0.0f )
+						.Visibility_Lambda( [&MeshEditorMode, WidgetSwitcher]() 
+							// Only show the widget switcher if either nothing is selected, or we have at least one mesh element selected
+							{ 
+								return ( MeshEditorMode.GetSelectedEditableMeshes().Num() == 0 || 
+										 WidgetSwitcher->GetActiveWidgetIndex() > static_cast<int32>( EEditableMeshElementType::Invalid ) ) ? 
+											EVisibility::Visible : EVisibility::Collapsed; 
+							} )
+						[
+							WidgetSwitcher
+						]
+					]
+					+SVerticalBox::Slot()
+					.AutoHeight()
+					.Padding( 6.0f )
+					[
+						SNew( SSeparator )
+						.Visibility_Lambda( [&MeshEditorMode, WidgetSwitcher]() 
+							// Only show the separator if we have a polygon, vertex or edge selected
+							{ 
+								return ( MeshEditorMode.GetSelectedEditableMeshes().Num() > 0 && 
+										 WidgetSwitcher->GetActiveWidgetIndex() > static_cast<int32>( EEditableMeshElementType::Invalid ) ) ? 
+											EVisibility::Visible : EVisibility::Collapsed; 
+							} )
+						.Orientation( Orient_Horizontal )
+					]
+					+ SVerticalBox::Slot()
+					.AutoHeight()
+					[
+						SNew( SBox )
+						.Visibility_Lambda( [&MeshEditorMode]() { return MeshEditorMode.GetSelectedEditableMeshes().Num() > 0 ? EVisibility::Visible : EVisibility::Collapsed; } )
+						[
+							SNew( SMeshEditorModeControlWidget, LOCTEXT( "MeshGroupName", "Mesh" ), MeshEditorMode.GetCommonActions() )
+						]
+					]
 				]
 			]
 		]
diff --git a/Engine/Plugins/Editor/MeshEditor/Source/MeshEditor/MeshEditorStyle.cpp b/Engine/Plugins/Editor/MeshEditor/Source/MeshEditor/MeshEditorStyle.cpp
--- a/Engine/Plugins/Editor/MeshEditor/Source/MeshEditor/MeshEditorStyle.cpp
+++ b/Engine/Plugins/Editor/MeshEditor/Source/MeshEditor/MeshEditorStyle.cpp
@@ -3,7 +3,7 @@
 #include "MeshEditorStyle.h"
 #include "Styling/SlateStyleRegistry.h"
 #include "IPluginManager.h"
-
+#include "SlateTypes.h"
 
 #define IMAGE_BRUSH( RelativePath, ... ) FSlateImageBrush( StyleSet->RootToContentDir( RelativePath, TEXT( ".png" ) ), __VA_ARGS__ )
 #define IMAGE_PLUGIN_BRUSH( RelativePath, ... ) FSlateImageBrush( FMeshEditorStyle::InContent( RelativePath, ".png" ), __VA_ARGS__ )
@@ -44,6 +44,12 @@ void FMeshEditorStyle::Initialize()
 	StyleSet->Set( "LevelEditor.MeshEditorMode.Selected", new IMAGE_PLUGIN_BRUSH( "Icons/MeshEditorMode_40px", Icon40x40 ) );
 	StyleSet->Set( "LevelEditor.MeshEditorMode.Selected.Small", new IMAGE_PLUGIN_BRUSH( "Icons/MeshEditorMode_40px", Icon40x40 ) );
 
+	StyleSet->Set( "EditingMode.GroupName.Text", FTextBlockStyle()
+		.SetFont( TTF_CORE_FONT( "Fonts/Roboto-Bold", 12 ) )
+		.SetColorAndOpacity( FLinearColor::White )
+		.SetHighlightColor( FLinearColor( 0.02f, 0.3f, 0.0f ) )
+		.SetHighlightShape( BOX_BRUSH( "Common/TextBlockHighlightShape", FMargin(3.f/8.f) ) ) );
+
 	StyleSet->Set( "EditingMode.Entry", FCheckBoxStyle()
 		.SetCheckBoxType( ESlateCheckBoxType::ToggleButton )
 		.SetUncheckedImage( BOX_BRUSH( "Common/Selection", 8.0f / 32.0f, DimBackground ) )
diff --git a/Engine/Plugins/Editor/MeshEditor/Source/PolygonModeling/BevelOrInsetPolygon.cpp b/Engine/Plugins/Editor/MeshEditor/Source/PolygonModeling/BevelOrInsetPolygon.cpp
--- a/Engine/Plugins/Editor/MeshEditor/Source/PolygonModeling/BevelOrInsetPolygon.cpp
+++ b/Engine/Plugins/Editor/MeshEditor/Source/PolygonModeling/BevelOrInsetPolygon.cpp
@@ -256,7 +256,7 @@ namespace BevelOrInsetPolygonHelpers
 
 void UBevelPolygonCommand::RegisterUICommand( FBindingContext* BindingContext )
 {
-	UI_COMMAND_EXT( BindingContext, /* Out */ UICommandInfo, "BevelPolygon", "Bevel Polygon Mode", "Set the primary action to bevel polygons.", EUserInterfaceActionType::RadioButton, FInputChord() );
+	UI_COMMAND_EXT( BindingContext, /* Out */ UICommandInfo, "BevelPolygon", "Bevel", "Bevels selected polygons as you click and drag.", EUserInterfaceActionType::RadioButton, FInputChord() );
 }
 
 
@@ -291,7 +291,7 @@ void UBevelPolygonCommand::AddToVRRadialMenuActionsMenu( IMeshEditorModeUIContra
 
 void UInsetPolygonCommand::RegisterUICommand( FBindingContext* BindingContext )
 {
-	UI_COMMAND_EXT( BindingContext, /* Out */ UICommandInfo, "InsetPolygon", "Inset Polygon Mode", "Set the primary action to inset polygons.", EUserInterfaceActionType::RadioButton, FInputChord() );
+	UI_COMMAND_EXT( BindingContext, /* Out */ UICommandInfo, "InsetPolygon", "Inset", "Insets selected polygons as you click and drag on the interior of a polygon.", EUserInterfaceActionType::RadioButton, FInputChord() );
 }
 
 
diff --git a/Engine/Plugins/Editor/MeshEditor/Source/PolygonModeling/EditVertexOrEdgeSharpness.cpp b/Engine/Plugins/Editor/MeshEditor/Source/PolygonModeling/EditVertexOrEdgeSharpness.cpp
--- a/Engine/Plugins/Editor/MeshEditor/Source/PolygonModeling/EditVertexOrEdgeSharpness.cpp
+++ b/Engine/Plugins/Editor/MeshEditor/Source/PolygonModeling/EditVertexOrEdgeSharpness.cpp
@@ -81,7 +81,7 @@ namespace VertexOrEdgeSharpnessHelpers
 
 void UEditVertexCornerSharpnessCommand::RegisterUICommand( FBindingContext* BindingContext )
 {
-	UI_COMMAND_EXT( BindingContext, /* Out */ UICommandInfo, "EditVertexCornerSharpness", "Edit Vertex Corner Sharpness Mode", "Set the primary action to edit vertex corner sharpness.", EUserInterfaceActionType::RadioButton, FInputChord() );
+	UI_COMMAND_EXT( BindingContext, /* Out */ UICommandInfo, "EditVertexCornerSharpness", "Edit Corner Sharpness", "Change the subdivision vertex corner sharpness of a vertex by clicking and dragging up and down.", EUserInterfaceActionType::RadioButton, FInputChord() );
 }
 
 
@@ -136,7 +136,7 @@ void UEditVertexCornerSharpnessCommand::ApplyDuringDrag( IMeshEditorModeEditingC
 
 void UEditEdgeCreaseSharpnessCommand::RegisterUICommand( FBindingContext* BindingContext )
 {
-	UI_COMMAND_EXT( BindingContext, /* Out */ UICommandInfo, "EditEdgeCreaseSharpness", "Edit Edge Crease Sharpness Mode", "Set the primary action to edit edge crease sharpness.", EUserInterfaceActionType::RadioButton, FInputChord() );
+	UI_COMMAND_EXT( BindingContext, /* Out */ UICommandInfo, "EditEdgeCreaseSharpness", "Edit Crease Sharpness", "Change the subdivision edge crase sharpness of an edge by clicking and dragging up and down.", EUserInterfaceActionType::RadioButton, FInputChord() );
 }
 
 
diff --git a/Engine/Plugins/Editor/MeshEditor/Source/PolygonModeling/ExtendEdge.cpp b/Engine/Plugins/Editor/MeshEditor/Source/PolygonModeling/ExtendEdge.cpp
--- a/Engine/Plugins/Editor/MeshEditor/Source/PolygonModeling/ExtendEdge.cpp
+++ b/Engine/Plugins/Editor/MeshEditor/Source/PolygonModeling/ExtendEdge.cpp
@@ -15,7 +15,7 @@
 
 void UExtendEdgeCommand::RegisterUICommand( FBindingContext* BindingContext )
 {
-	UI_COMMAND_EXT( BindingContext, /* Out */ UICommandInfo, "ExtendEdge", "Extend Edge Mode", "Set the primary action to extend edges.", EUserInterfaceActionType::RadioButton, FInputChord() );
+	UI_COMMAND_EXT( BindingContext, /* Out */ UICommandInfo, "ExtendEdge", "Extend", "Creates a new polygon from the selected edge by clicking and dragging outward from an edge.", EUserInterfaceActionType::RadioButton, FInputChord() );
 }
 
 
diff --git a/Engine/Plugins/Editor/MeshEditor/Source/PolygonModeling/ExtendVertex.cpp b/Engine/Plugins/Editor/MeshEditor/Source/PolygonModeling/ExtendVertex.cpp
--- a/Engine/Plugins/Editor/MeshEditor/Source/PolygonModeling/ExtendVertex.cpp
+++ b/Engine/Plugins/Editor/MeshEditor/Source/PolygonModeling/ExtendVertex.cpp
@@ -15,7 +15,7 @@
 
 void UExtendVertexCommand::RegisterUICommand( FBindingContext* BindingContext )
 {
-	UI_COMMAND_EXT( BindingContext, /* Out */ UICommandInfo, "ExtendVertex", "Extend Vertex Mode", "Set the primary action to extend vertices.", EUserInterfaceActionType::RadioButton, FInputChord() );
+	UI_COMMAND_EXT( BindingContext, /* Out */ UICommandInfo, "ExtendVertex", "Extend", "Creates a new triangle from the selected vertex by clicking and dragging outward from the vertex.  The new triangle will connect to the next closest neighbor vertex.", EUserInterfaceActionType::RadioButton, FInputChord() );
 }
 
 
diff --git a/Engine/Plugins/Editor/MeshEditor/Source/PolygonModeling/ExtrudePolygon.cpp b/Engine/Plugins/Editor/MeshEditor/Source/PolygonModeling/ExtrudePolygon.cpp
--- a/Engine/Plugins/Editor/MeshEditor/Source/PolygonModeling/ExtrudePolygon.cpp
+++ b/Engine/Plugins/Editor/MeshEditor/Source/PolygonModeling/ExtrudePolygon.cpp
@@ -71,7 +71,7 @@ void UExtrudePolygonCommand::FindExtrudeDistanceUsingInteractor(
 
 void UExtrudePolygonCommand::RegisterUICommand( FBindingContext* BindingContext )
 {
-	UI_COMMAND_EXT( BindingContext, /* Out */ UICommandInfo, "ExtrudePolygon", "Extrude Polygon Mode", "Set the primary action to extrude polygons.", EUserInterfaceActionType::RadioButton, FInputChord() );
+	UI_COMMAND_EXT( BindingContext, /* Out */ UICommandInfo, "ExtrudePolygon", "Extrude", "Pull the selected polygons along their surface normal, creating new quads along along the side of each edge.", EUserInterfaceActionType::RadioButton, FInputChord() );
 }
 
 
diff --git a/Engine/Plugins/Editor/MeshEditor/Source/PolygonModeling/HardenOrSoftenEdge.cpp b/Engine/Plugins/Editor/MeshEditor/Source/PolygonModeling/HardenOrSoftenEdge.cpp
--- a/Engine/Plugins/Editor/MeshEditor/Source/PolygonModeling/HardenOrSoftenEdge.cpp
+++ b/Engine/Plugins/Editor/MeshEditor/Source/PolygonModeling/HardenOrSoftenEdge.cpp
@@ -5,7 +5,7 @@
 #include "UICommandInfo.h"
 #include "EditableMesh.h"
 #include "MeshElement.h"
-
+#include "ScopedTransaction.h"
 
 #define LOCTEXT_NAMESPACE "MeshEditorMode"
 
@@ -69,7 +69,7 @@ namespace HardenOrSoftenEdgeHelpers
 
 void UHardenEdgeCommand::RegisterUICommand( FBindingContext* BindingContext )
 {
-	UI_COMMAND_EXT( BindingContext, /* Out */ UICommandInfo, "HardenEdge", "Harden Edge", "Sets the edge to be hard.", EUserInterfaceActionType::Button, FInputChord( EKeys::H ) );
+	UI_COMMAND_EXT( BindingContext, /* Out */ UICommandInfo, "HardenEdge", "Harden", "Sets the edge to be hard.", EUserInterfaceActionType::Button, FInputChord( EKeys::H ) );
 }
 
 
@@ -82,7 +82,7 @@ void UHardenEdgeCommand::Execute( IMeshEditorModeEditingContract& MeshEditorMode
 
 void USoftenEdgeCommand::RegisterUICommand( FBindingContext* BindingContext )
 {
-	UI_COMMAND_EXT( BindingContext, /* Out */ UICommandInfo, "SoftenEdge", "Soften Edge", "Sets the edge to be soft.", EUserInterfaceActionType::Button, FInputChord( EKeys::H, EModifierKey::Shift ) );
+	UI_COMMAND_EXT( BindingContext, /* Out */ UICommandInfo, "SoftenEdge", "Soften", "Sets the edge to be soft.", EUserInterfaceActionType::Button, FInputChord( EKeys::H, EModifierKey::Shift ) );
 }
 
 
diff --git a/Engine/Plugins/Editor/MeshEditor/Source/PolygonModeling/InsertEdgeLoop.cpp b/Engine/Plugins/Editor/MeshEditor/Source/PolygonModeling/InsertEdgeLoop.cpp
--- a/Engine/Plugins/Editor/MeshEditor/Source/PolygonModeling/InsertEdgeLoop.cpp
+++ b/Engine/Plugins/Editor/MeshEditor/Source/PolygonModeling/InsertEdgeLoop.cpp
@@ -15,7 +15,7 @@
 
 void UInsertEdgeLoopCommand::RegisterUICommand( FBindingContext* BindingContext )
 {
-	UI_COMMAND_EXT( BindingContext, /* Out */ UICommandInfo, "InsertEdgeLoop", "Insert Edge Loop Mode", "Set the primary action to insert edge loops.", EUserInterfaceActionType::RadioButton, FInputChord() );
+	UI_COMMAND_EXT( BindingContext, /* Out */ UICommandInfo, "InsertEdgeLoop", "Insert Edge Loop", "Inserts a loop of edges at a specific location along the selected edge as you click and drag.  If a valid loop cannot be determined, no edges will be inserted.", EUserInterfaceActionType::RadioButton, FInputChord() );
 }
 
 
diff --git a/Engine/Plugins/Editor/MeshEditor/Source/PolygonModeling/QuadrangulateMesh.cpp b/Engine/Plugins/Editor/MeshEditor/Source/PolygonModeling/QuadrangulateMesh.cpp
--- a/Engine/Plugins/Editor/MeshEditor/Source/PolygonModeling/QuadrangulateMesh.cpp
+++ b/Engine/Plugins/Editor/MeshEditor/Source/PolygonModeling/QuadrangulateMesh.cpp
@@ -13,7 +13,7 @@
 
 void UQuadrangulateMeshCommand::RegisterUICommand( FBindingContext* BindingContext )
 {
-	UI_COMMAND_EXT( BindingContext, /* Out */ UICommandInfo, "QuadrangulateMesh", "Quadrangulate Mesh", "Quadrangulates the selected mesh.", EUserInterfaceActionType::Button, FInputChord() );
+	UI_COMMAND_EXT( BindingContext, /* Out */ UICommandInfo, "QuadrangulateMesh", "Quadrangulate", "Quadrangulates the selected mesh.", EUserInterfaceActionType::Button, FInputChord() );
 }
 
 
diff --git a/Engine/Plugins/Editor/MeshEditor/Source/PolygonModeling/RemoveEdge.cpp b/Engine/Plugins/Editor/MeshEditor/Source/PolygonModeling/RemoveEdge.cpp
--- a/Engine/Plugins/Editor/MeshEditor/Source/PolygonModeling/RemoveEdge.cpp
+++ b/Engine/Plugins/Editor/MeshEditor/Source/PolygonModeling/RemoveEdge.cpp
@@ -14,7 +14,7 @@
 
 void URemoveEdgeCommand::RegisterUICommand( FBindingContext* BindingContext )
 {
-	UI_COMMAND_EXT( BindingContext, /* Out */ UICommandInfo, "RemoveEdge", "Remove Edge", "Attempts to remove the selected edge and merge adjacent polygons.", EUserInterfaceActionType::Button, FInputChord( EKeys::BackSpace ) );
+	UI_COMMAND_EXT( BindingContext, /* Out */ UICommandInfo, "RemoveEdge", "Remove", "Attempts to remove the selected edge and merge adjacent polygons.", EUserInterfaceActionType::Button, FInputChord( EKeys::BackSpace ) );
 }
 
 
diff --git a/Engine/Plugins/Editor/MeshEditor/Source/PolygonModeling/RemoveVertex.cpp b/Engine/Plugins/Editor/MeshEditor/Source/PolygonModeling/RemoveVertex.cpp
--- a/Engine/Plugins/Editor/MeshEditor/Source/PolygonModeling/RemoveVertex.cpp
+++ b/Engine/Plugins/Editor/MeshEditor/Source/PolygonModeling/RemoveVertex.cpp
@@ -6,6 +6,9 @@
 #include "UICommandInfo.h"
 #include "EditableMesh.h"
 #include "MeshElement.h"
+#include "ScopedTransaction.h"
+#include "MultiBoxBuilder.h"
+#include "UICommandList.h"
 
 
 #define LOCTEXT_NAMESPACE "MeshEditorMode"
@@ -13,7 +16,7 @@
 
 void URemoveVertexCommand::RegisterUICommand( FBindingContext* BindingContext )
 {
-	UI_COMMAND_EXT( BindingContext, /* Out */ UICommandInfo, "RemoveVertex", "Remove Vertex", "Attempts to remove the selected vertex, keeping the polygon intact.", EUserInterfaceActionType::Button, FInputChord( EKeys::BackSpace ) );
+	UI_COMMAND_EXT( BindingContext, /* Out */ UICommandInfo, "RemoveVertex", "Remove", "Attempts to remove the selected vertex, keeping the polygon intact.", EUserInterfaceActionType::Button, FInputChord( EKeys::BackSpace ) );
 }
 
 
diff --git a/Engine/Plugins/Editor/MeshEditor/Source/PolygonModeling/SplitEdge.cpp b/Engine/Plugins/Editor/MeshEditor/Source/PolygonModeling/SplitEdge.cpp
--- a/Engine/Plugins/Editor/MeshEditor/Source/PolygonModeling/SplitEdge.cpp
+++ b/Engine/Plugins/Editor/MeshEditor/Source/PolygonModeling/SplitEdge.cpp
@@ -15,7 +15,7 @@
 
 void USplitEdgeCommand::RegisterUICommand( FBindingContext* BindingContext )
 {
-	UI_COMMAND_EXT( BindingContext, /* Out */ UICommandInfo, "SplitEdge", "Split Edge Mode", "Set the primary action to split edges.", EUserInterfaceActionType::RadioButton, FInputChord() );
+	UI_COMMAND_EXT( BindingContext, /* Out */ UICommandInfo, "InsertVertex", "Insert Vertex", "Inserts a vertex at a specific position along an edge as you click and drag, splitting the edge into two.", EUserInterfaceActionType::RadioButton, FInputChord() );
 }
 
 
diff --git a/Engine/Plugins/Editor/MeshEditor/Source/PolygonModeling/SplitPolygon.cpp b/Engine/Plugins/Editor/MeshEditor/Source/PolygonModeling/SplitPolygon.cpp
--- a/Engine/Plugins/Editor/MeshEditor/Source/PolygonModeling/SplitPolygon.cpp
+++ b/Engine/Plugins/Editor/MeshEditor/Source/PolygonModeling/SplitPolygon.cpp
@@ -15,19 +15,19 @@
 
 void USplitPolygonFromVertexCommand::RegisterUICommand( FBindingContext* BindingContext )
 {
-	UI_COMMAND_EXT( BindingContext, /* Out */ UICommandInfo, "SplitPolygonFromVertex", "Split Polygon Mode", "Set the primary action to split a polygon.", EUserInterfaceActionType::RadioButton, FInputChord() );
+	UI_COMMAND_EXT( BindingContext, /* Out */ UICommandInfo, "SplitPolygonFromVertex", "Split Polygon", "Splits a polygon by clicking on a selected vertex and dragging to create an edge along the surface of a neighboring polygon.", EUserInterfaceActionType::RadioButton, FInputChord() );
 }
 
 
 void USplitPolygonFromEdgeCommand::RegisterUICommand( FBindingContext* BindingContext )
 {
-	UI_COMMAND_EXT( BindingContext, /* Out */ UICommandInfo, "SplitPolygonFromEdge", "Split Polygon Mode", "Set the primary action to split a polygon.", EUserInterfaceActionType::RadioButton, FInputChord() );
+	UI_COMMAND_EXT( BindingContext, /* Out */ UICommandInfo, "SplitPolygonFromEdge", "Split Polygon", "Splits a polygon by clicking on a selected edge and dragging to create an edge along the surface of a neigboring polygon.", EUserInterfaceActionType::RadioButton, FInputChord() );
 }
 
 
 void USplitPolygonFromPolygonCommand::RegisterUICommand( FBindingContext* BindingContext )
 {
-	UI_COMMAND_EXT( BindingContext, /* Out */ UICommandInfo, "SplitPolygonFromPolygon", "Split Polygon Mode", "Set the primary action to split a polygon.", EUserInterfaceActionType::RadioButton, FInputChord() );
+	UI_COMMAND_EXT( BindingContext, /* Out */ UICommandInfo, "SplitPolygonFromPolygon", "Split", "Splits a polygon by clicking on a selected polygon and dragging to create an edge along the surface of the polygon or it's neighbor.", EUserInterfaceActionType::RadioButton, FInputChord() );
 }
 
 
diff --git a/Engine/Plugins/Editor/MeshEditor/Source/PolygonModeling/TessellatePolygon.cpp b/Engine/Plugins/Editor/MeshEditor/Source/PolygonModeling/TessellatePolygon.cpp
--- a/Engine/Plugins/Editor/MeshEditor/Source/PolygonModeling/TessellatePolygon.cpp
+++ b/Engine/Plugins/Editor/MeshEditor/Source/PolygonModeling/TessellatePolygon.cpp
@@ -13,7 +13,7 @@
 
 void UTessellatePolygonCommand::RegisterUICommand( FBindingContext* BindingContext )
 {
-	UI_COMMAND_EXT( BindingContext, /* Out */ UICommandInfo, "TessellatePolygon", "Tessellate Polygon", "Tessellate selected polygons into smaller polygons.", EUserInterfaceActionType::Button, FInputChord() );
+	UI_COMMAND_EXT( BindingContext, /* Out */ UICommandInfo, "TessellatePolygon", "Tessellate", "Tessellate selected polygons into smaller polygons.", EUserInterfaceActionType::Button, FInputChord() );
 }
 
 
