Mesh Editor: Added shortcut keys for actions and modes to their label
- Not super happy with how this looks yet, but just want to try it out.

diff --git a/Engine/Plugins/Editor/MeshEditor/Source/MeshEditor/MeshEditorModeToolkit.cpp b/Engine/Plugins/Editor/MeshEditor/Source/MeshEditor/MeshEditorModeToolkit.cpp
--- a/Engine/Plugins/Editor/MeshEditor/Source/MeshEditor/MeshEditorModeToolkit.cpp
+++ b/Engine/Plugins/Editor/MeshEditor/Source/MeshEditor/MeshEditorModeToolkit.cpp
@@ -47,8 +47,10 @@ public:
 					SNew( SButton )
 					.HAlign( HAlign_Center )
 					.VAlign( VAlign_Center )
-					.ContentPadding( FMargin( 4.0f ) )
-					.Text( CommandInfo.GetLabel() )
+					.ContentPadding( FMargin( 8.0f, 4.0f ) )
+					.Text( CommandInfo.GetDefaultChord().IsValidChord() ?
+						FText::Format( LOCTEXT( "ButtonLabelAndShortcutFormat", "{0}  ({1})" ), CommandInfo.GetLabel(), CommandInfo.GetDefaultChord().GetInputText() ) :
+						CommandInfo.GetLabel() )
 					.ToolTip( SNew( SToolTip ).Text( CommandInfo.GetDescription() ) )
 					.OnClicked_Lambda( [UIAction] { UIAction.Execute(); return FReply::Handled(); } )
 					.IsEnabled_Lambda( [UIAction] { return UIAction.CanExecute(); } )
@@ -75,13 +77,15 @@ public:
 							.Size( FVector2D( 1, 30 ) )
 						]
 						+SOverlay::Slot()
-						.Padding( FMargin( 20, 0, 20, 0 ) )
+						.Padding( FMargin( 8, 0, 8, 0 ) )
 						.HAlign( HAlign_Center )
 						.VAlign( VAlign_Center )
 						[
 							SNew( STextBlock )
 							.TextStyle( FMeshEditorStyle::Get(), "EditingMode.Entry.Text" )
-							.Text( CommandInfo.GetLabel() )
+							.Text( CommandInfo.GetDefaultChord().IsValidChord() ? 
+								FText::Format( LOCTEXT( "RadioButtonLabelAndShortcutFormat", "{0}  ({1})" ), CommandInfo.GetLabel(), CommandInfo.GetDefaultChord().GetInputText() ) :
+								CommandInfo.GetLabel() )
 						]
 					]
 				];
