Persistent wireframe meshes using custom components.

Instead of rendering the overlay UI using the PDI and a dynamic mesh builder, UI elements are now rendered via a couple of custom components. This has two advantages:

1) We can specify a sort priority for rendering and hence ensure that transparent elements are rendered in the correct order relative to the mesh they are overlaying, and each other.
2) The overlay meshes can be persistent, meaning they need only be updated when the mesh they reference changes, or some other change occurs (e.g. new selected or hovered element).

There are two different components, with subtly different uses.

UWireframeMeshComponent is used to instance a persistent UWireframeMesh. UWireframeMesh maps 1:1 with an editable/static mesh, and there can be many components referencing it.  In this case, it is used to provide a persistent wireframe mesh for the editable mesh base cage, and another for the subdivided mesh wireframe cage, two such instances per editable mesh. Two new adapter classes are provided to bridge between the editable mesh and the wireframe meshes, so they can be updated incrementally.

UOverlayComponent does not reference a mesh asset, but holds all the mesh data itself.  The intention is for UOverlayComponent to become a generic engine component if we wish, hence it is not tied to the editable mesh format.  The mesh editor creates three UOverlayComponents, which act across the entire level rather than per editable mesh: one provides hovered element rendering, another selected element rendering, and another renders selected edges on the subdivided mesh cage.

UOverlayComponent can hold three element types: points, lines and triangles, and custom materials can be provided for each.  Camera-facing elements (points and lines) are handled by the material, where the UV0 and vertex normal channels have special meanings used to determine the size/thickness of the rendered elements.

All mesh editor UI rendering has been replaced by this component-based system.

diff --git a/Engine/Plugins/Editor/MeshEditor/Source/MeshEditor/MeshEditorAssetContainer.h b/Engine/Plugins/Editor/MeshEditor/Source/MeshEditor/MeshEditorAssetContainer.h
--- a/Engine/Plugins/Editor/MeshEditor/Source/MeshEditor/MeshEditorAssetContainer.h
+++ b/Engine/Plugins/Editor/MeshEditor/Source/MeshEditor/MeshEditorAssetContainer.h
@@ -14,10 +14,22 @@ class MESHEDITOR_API UMeshEditorAssetContainer : public UDataAsset
  
 public:
 	UPROPERTY(EditAnywhere, Category = Material)
-	class UMaterial* HoveredGeometryMaterial;
+	class UMaterialInterface* HoveredGeometryMaterial;
 
 	UPROPERTY(EditAnywhere, Category = Material)
-	class UMaterial* HoveredFaceMaterial;
+	class UMaterialInterface* HoveredFaceMaterial;
+
+	UPROPERTY(EditAnywhere, Category = Material)
+	class UMaterialInterface* WireMaterial;
+
+	UPROPERTY(EditAnywhere, Category = Material)
+	class UMaterialInterface* SubdividedMeshWireMaterial;
+
+	UPROPERTY(EditAnywhere, Category = Material)
+	class UMaterialInterface* OverlayLineMaterial;
+
+	UPROPERTY(EditAnywhere, Category = Material)
+	class UMaterialInterface* OverlayPointMaterial;
 
 	UPROPERTY(EditAnywhere, Category = Sound)
 	class USoundBase* DefaultSound;
diff --git a/Engine/Plugins/Editor/MeshEditor/Source/MeshEditor/MeshEditorMode.cpp b/Engine/Plugins/Editor/MeshEditor/Source/MeshEditor/MeshEditorMode.cpp
--- a/Engine/Plugins/Editor/MeshEditor/Source/MeshEditor/MeshEditorMode.cpp
+++ b/Engine/Plugins/Editor/MeshEditor/Source/MeshEditor/MeshEditorMode.cpp
@@ -37,7 +37,9 @@
 #include "SlateApplication.h"	// For GetCurrentTime()
 #include "Materials/Material.h"
 #include "Materials/MaterialInterface.h"
-
+#include "MeshEditorStaticMeshAdapter.h"
+#include "MeshEditorSubdividedStaticMeshAdapter.h"
+#include "WireframeMeshComponent.h"
 #include "Misc/FeedbackContext.h"
 #include "Framework/Application/SlateApplication.h"
 #include "Materials/Material.h"
@@ -393,6 +395,7 @@ FMeshEditorMode::FMeshEditorInteractorData::FMeshEditorInteractorData()
 FMeshEditorMode::FMeshEditorMode()
 	: HoveredGeometryMaterial( nullptr ),
 	  HoveredFaceMaterial( nullptr ),
+	  WireMaterial( nullptr ),
 	  HoverFeedbackTimeValue( 0.0 ),
 	  MeshElementSelectionMode( EEditableMeshElementType::Any ),
 	  EquippedVertexAction( EMeshEditAction::Move ),
@@ -410,10 +413,11 @@ FMeshEditorMode::FMeshEditorMode()
 	  bShowVertexNormals( false ),
 	  bMarqueeSelectTransactionActive( false ),
 	  bShouldFocusToSelection( false ),
+	  bShouldUpdateSelectedElementsOverlay( false ),
 	  bPerInstanceEdits( false ),
 	  AssetContainer( nullptr )
 {
-	AssetContainer = LoadObject<UMeshEditorAssetContainer>(nullptr, TEXT("/MeshEditor/MeshEditorAssetContainer"));
+	AssetContainer = LoadObject<UMeshEditorAssetContainer>( nullptr, TEXT( "/MeshEditor/MeshEditorAssetContainer" ) );
 	check(AssetContainer != nullptr);
 
 	HoveredGeometryMaterial = AssetContainer->HoveredGeometryMaterial;
@@ -422,6 +426,18 @@ FMeshEditorMode::FMeshEditorMode()
 	HoveredFaceMaterial = AssetContainer->HoveredFaceMaterial;
 	check( HoveredFaceMaterial != nullptr );
 
+	WireMaterial = AssetContainer->WireMaterial;
+	check( WireMaterial != nullptr );
+
+	OverlayLineMaterial = AssetContainer->OverlayLineMaterial;
+	check( OverlayLineMaterial != nullptr );
+
+	OverlayPointMaterial = AssetContainer->OverlayPointMaterial;
+	check( OverlayPointMaterial != nullptr );
+
+	SubdividedMeshWireMaterial = AssetContainer->SubdividedMeshWireMaterial;
+	check( SubdividedMeshWireMaterial != nullptr );
+
 	MeshEditorModeProxyObject = NewObject<UMeshEditorModeProxyObject>();
 	MeshEditorModeProxyObject->OwningMeshEditorMode = this;
 
@@ -444,6 +460,14 @@ FMeshEditorMode::FMeshEditorMode()
 }
 
 
+void FMeshEditorMode::Initialize()
+{
+	FActorSpawnParameters ActorSpawnParameters;
+	ActorSpawnParameters.ObjectFlags |= RF_Transient;
+	WireframeComponentContainer = GetWorld()->SpawnActor<AActor>( ActorSpawnParameters );
+}
+
+
 FMeshEditorMode::~FMeshEditorMode()
 {
 	FEditorDelegates::EndPIE.RemoveAll( this );
@@ -465,9 +489,12 @@ FMeshEditorMode::~FMeshEditorMode()
 	// Remove the event registered on all cached editable meshes
 	for( auto& CachedEditableMesh : CachedEditableMeshes )
 	{
-		CachedEditableMesh.Value->OnElementIDsRemapped().RemoveAll( this );
+		CachedEditableMesh.Value.EditableMesh->OnElementIDsRemapped().RemoveAll( this );
 	}
 
+	WireframeComponentContainer->Destroy();
+	WireframeComponentContainer = nullptr;
+
 	MeshEditorModeProxyObject = nullptr;
 	AssetContainer = nullptr;
 }
@@ -496,22 +523,25 @@ void FMeshEditorMode::OnEditableMeshElementIDsRemapped( UEditableMesh* EditableM
 	// Helper function which performs the remapping of a given FMeshElement
 	auto RemapMeshElement = [ this, EditableMesh, &Remappings ]( FMeshElement& MeshElement )
 	{
-		UEditableMesh* MeshElementEditableMesh = this->FindOrCreateEditableMesh( *MeshElement.Component, MeshElement.ElementAddress.SubMeshAddress );
-		if( MeshElementEditableMesh == EditableMesh )
+		if( MeshElement.Component.IsValid() )
 		{
-			switch( MeshElement.ElementAddress.ElementType )
+			UEditableMesh* MeshElementEditableMesh = this->FindOrCreateEditableMesh( *MeshElement.Component, MeshElement.ElementAddress.SubMeshAddress );
+			if( MeshElementEditableMesh == EditableMesh )
 			{
-			case EEditableMeshElementType::Vertex:
-				MeshElement.ElementAddress.ElementID = Remappings.GetRemappedVertexID( FVertexID( MeshElement.ElementAddress.ElementID ) );
-				break;
+				switch( MeshElement.ElementAddress.ElementType )
+				{
+				case EEditableMeshElementType::Vertex:
+					MeshElement.ElementAddress.ElementID = Remappings.GetRemappedVertexID( FVertexID( MeshElement.ElementAddress.ElementID ) );
+					break;
 
-			case EEditableMeshElementType::Edge:
-				MeshElement.ElementAddress.ElementID = Remappings.GetRemappedEdgeID( FEdgeID( MeshElement.ElementAddress.ElementID ) );
-				break;
+				case EEditableMeshElementType::Edge:
+					MeshElement.ElementAddress.ElementID = Remappings.GetRemappedEdgeID( FEdgeID( MeshElement.ElementAddress.ElementID ) );
+					break;
 
-			case EEditableMeshElementType::Polygon:
-				MeshElement.ElementAddress.ElementID = Remappings.GetRemappedPolygonID( FPolygonID( MeshElement.ElementAddress.ElementID ) );
-				break;
+				case EEditableMeshElementType::Polygon:
+					MeshElement.ElementAddress.ElementID = Remappings.GetRemappedPolygonID( FPolygonID( MeshElement.ElementAddress.ElementID ) );
+					break;
+				}
 			}
 		}
 	};
@@ -801,6 +831,34 @@ void FMeshEditorMode::Enter()
 	// Call parent implementation
 	FEdMode::Enter();
 
+	// Add overlay component for rendering hovered elements
+	HoveredElementsComponent= NewObject<UOverlayComponent>( WireframeComponentContainer );
+	HoveredElementsComponent->SetLineMaterial( OverlayLineMaterial );
+	HoveredElementsComponent->SetPointMaterial( OverlayPointMaterial );
+	HoveredElementsComponent->TranslucencySortPriority = 400;
+	HoveredElementsComponent->RegisterComponent();
+
+	// Add overlay component for rendering selected elements
+	SelectedElementsComponent = NewObject<UOverlayComponent>( WireframeComponentContainer );
+	SelectedElementsComponent->SetLineMaterial( OverlayLineMaterial );
+	SelectedElementsComponent->SetPointMaterial( OverlayPointMaterial );
+	SelectedElementsComponent->TranslucencySortPriority = 500;
+	SelectedElementsComponent->RegisterComponent();
+
+	// Add overlay component for rendering selected wires on the SubD mesh
+	SelectedSubDElementsComponent = NewObject<UOverlayComponent>( WireframeComponentContainer );
+	SelectedSubDElementsComponent->SetLineMaterial( OverlayLineMaterial );
+	SelectedSubDElementsComponent->SetPointMaterial( OverlayPointMaterial );
+	SelectedSubDElementsComponent->TranslucencySortPriority = 200;
+	SelectedSubDElementsComponent->RegisterComponent();
+
+	// Add overlay component for rendering debug normals/tangents on the base cage
+	DebugNormalsComponent = NewObject<UOverlayComponent>( WireframeComponentContainer );
+	DebugNormalsComponent->SetLineMaterial( OverlayLineMaterial );
+	DebugNormalsComponent->SetPointMaterial( OverlayPointMaterial );
+	DebugNormalsComponent->TranslucencySortPriority = 600;
+	DebugNormalsComponent->RegisterComponent();
+
 	UEditorWorldExtensionCollection* ExtensionCollection = GEditor->GetEditorWorldExtensionsManager()->GetEditorWorldExtensions( GetWorld() );
 	check( ExtensionCollection != nullptr );
 
@@ -895,6 +953,11 @@ void FMeshEditorMode::Exit()
 		GEditor->NoteSelectionChange();
 	}
 
+	// Remove overlay components
+	SelectedSubDElementsComponent->DestroyComponent();
+	SelectedElementsComponent->DestroyComponent();
+	HoveredElementsComponent->DestroyComponent();
+
 	// Call parent implementation
 	FEdMode::Exit();
 }
@@ -904,11 +967,11 @@ const UEditableMesh* FMeshEditorMode::FindEditableMesh( UPrimitiveComponent& Com
 {
 	const UEditableMesh* EditableMesh = nullptr;
 
-	// Grab the existing editable mesh from our cache if we have one, otherwise create one now
-	UEditableMesh* const* EditableMeshPtr = CachedEditableMeshes.Find( SubMeshAddress );
-	if( EditableMeshPtr != nullptr )
+	// Grab the existing editable mesh from our cache if we have one
+	const FEditableAndWireframeMeshes* EditableAndWireframeMeshes = CachedEditableMeshes.Find( SubMeshAddress );
+	if( EditableAndWireframeMeshes )
 	{
-		EditableMesh = *EditableMeshPtr;
+		EditableMesh = EditableAndWireframeMeshes->EditableMesh;
 	}
 
 	return EditableMesh;
@@ -920,10 +983,10 @@ UEditableMesh* FMeshEditorMode::FindOrCreateEditableMesh( UPrimitiveComponent& C
 	UEditableMesh* EditableMesh = nullptr;
 
 	// Grab the existing editable mesh from our cache if we have one, otherwise create one now
-	UEditableMesh** EditableMeshPtr = CachedEditableMeshes.Find( SubMeshAddress );
-	if( EditableMeshPtr != nullptr )
+	const FEditableAndWireframeMeshes* EditableAndWireframeMeshesPtr = CachedEditableMeshes.Find( SubMeshAddress );
+	if( EditableAndWireframeMeshesPtr )
 	{
-		EditableMesh = *EditableMeshPtr;
+		EditableMesh = EditableAndWireframeMeshesPtr->EditableMesh;
 	}
 	else
 	{
@@ -931,7 +994,6 @@ UEditableMesh* FMeshEditorMode::FindOrCreateEditableMesh( UPrimitiveComponent& C
 		{
 			// @todo mesheditor perf: This is going to HITCH as you hover over meshes.  Ideally we do this on a thread, or worst case give the user a progress dialog.  Maybe save out the editable mesh in editor builds?
 			EditableMesh = UEditableMeshFactory::MakeEditableMesh( &Component, SubMeshAddress );
-			CachedEditableMeshes.Add( SubMeshAddress, EditableMesh );
 
 			if( GetDefault<UMeshEditorSettings>()->bAutoQuadrangulate )
 			{
@@ -952,9 +1014,40 @@ UEditableMesh* FMeshEditorMode::FindOrCreateEditableMesh( UPrimitiveComponent& C
 			// Enable compaction on this mesh and set a callback so any cached ElementIDs can be remapped
 			EditableMesh->SetAllowCompact( true );
 			EditableMesh->OnElementIDsRemapped().AddRaw( this, &FMeshEditorMode::OnEditableMeshElementIDsRemapped );
+
+			// Create a wireframe mesh for the base cage
+			UWireframeMesh* WireframeBaseCage = NewObject<UWireframeMesh>();
+
+			UMeshEditorStaticMeshAdapter* WireframeAdapter = NewObject<UMeshEditorStaticMeshAdapter>();
+			EditableMesh->Adapters.Add( WireframeAdapter );
+			WireframeAdapter->Initialize( EditableMesh, WireframeBaseCage );
+
+			// Create a wireframe mesh for the subdivided mesh
+			UWireframeMesh* WireframeSubdividedMesh = NewObject<UWireframeMesh>();
+
+			UMeshEditorSubdividedStaticMeshAdapter* WireframeSubdividedAdapter = NewObject<UMeshEditorSubdividedStaticMeshAdapter>();
+			EditableMesh->Adapters.Add( WireframeSubdividedAdapter );
+			WireframeSubdividedAdapter->Initialize( EditableMesh, WireframeSubdividedMesh );
+
+			// Rebuild mesh so that the wireframe meshes get their render data built through the adapters
+			EditableMesh->RebuildRenderMesh();
+
+			// Cache the editable mesh and the associated wireframe meshes
+			FEditableAndWireframeMeshes EditableAndWireframeMeshes;
+			EditableAndWireframeMeshes.EditableMesh = EditableMesh;
+			EditableAndWireframeMeshes.WireframeBaseCage = WireframeBaseCage;
+			EditableAndWireframeMeshes.WireframeSubdividedMesh = WireframeSubdividedMesh;
+
+			CachedEditableMeshes.Add( SubMeshAddress, EditableAndWireframeMeshes );
 		}
 	}
 
+	// Create a wireframe component if necessary
+	const FWireframeMeshComponents& WireframeMeshComponents = CreateWireframeMeshComponents( &Component );
+	const FTransform Transform = Component.GetComponentTransform();
+	WireframeMeshComponents.WireframeMeshComponent->SetWorldTransform( Transform );
+	WireframeMeshComponents.WireframeSubdividedMeshComponent->SetWorldTransform( Transform );
+
 	return EditableMesh;
 }
 
@@ -1074,6 +1167,7 @@ void FMeshEditorMode::Tick( FEditorViewportClient* ViewportClient, float DeltaTi
 	{
 		const bool bIsActionFinishing = false;
 		UpdateActiveAction( bIsActionFinishing );
+
 	}
 
 
@@ -1098,6 +1192,225 @@ void FMeshEditorMode::Tick( FEditorViewportClient* ViewportClient, float DeltaTi
 				  ( EquippedEdgeAction == EMeshEditAction::Move && SelectedMeshElementType == EEditableMeshElementType::Edge ) ) ) 
 			);
 	}
+
+	// Update hovered/selected elements.
+	// @todo mesheditor: Ideally selected elements would be persistent and just updated when selection changes, or when geometry changes.
+	// There's currently not a simple way of doing the latter as there's no common path in the mesh editor for when mesh edits are performed.
+	// Potentially this could be done with another adapter, although it's a per-component thing rather than a per-editable mesh thing.
+
+	HoveredElementsComponent->Clear();
+//	SelectedElementsComponent->Clear();
+
+	const double CurrentRealTime = FSlateApplication::Get().GetCurrentTime();
+
+	// Only draw hover if we're not in the middle of an interactive edit
+	if( ActiveAction == NAME_None )
+	{
+		const float HoveredSizeBias = MeshEd::HoveredSizeBias->GetFloat() + MeshEd::HoveredAnimationExtraSizeBias->GetFloat() * FMath::MakePulsatingValue( HoverFeedbackTimeValue, 0.5f );
+
+		// Update hovered meshes
+		for( FMeshEditorInteractorData& MeshEditorInteractorData : MeshEditorInteractorDatas )
+		{
+			if( MeshElementSelectionMode == EEditableMeshElementType::Any ||
+				MeshEditorInteractorData.HoveredMeshElement.ElementAddress.ElementType == MeshElementSelectionMode )
+			{
+				FMeshElement HoveredMeshElement = GetHoveredMeshElement( MeshEditorInteractorData.ViewportInteractor.Get() );
+
+				const FColor Color = FLinearColor::Green.ToFColor( false );
+				AddMeshElementToOverlay( HoveredElementsComponent, HoveredMeshElement, Color, HoveredSizeBias );
+			}
+		}
+
+		// Update meshes that were previously hovered
+		const float HoverFadeTime = MeshEd::HoverFadeDuration->GetFloat();
+
+		for( FMeshElement& FadingOutHoveredMeshElement : FadingOutHoveredMeshElements )
+		{
+			const UEditableMesh* EditableMesh = FindEditableMesh( *FadingOutHoveredMeshElement.Component.Get(), FadingOutHoveredMeshElement.ElementAddress.SubMeshAddress );
+			if( IsElementIDValid( FadingOutHoveredMeshElement, EditableMesh ) )
+			{
+				const float TimeSinceLastHovered = CurrentRealTime - FadingOutHoveredMeshElement.LastHoverTime;
+				float Opacity = 1.0f - ( TimeSinceLastHovered / HoverFadeTime );
+				Opacity = Opacity * Opacity * Opacity * Opacity;		// Exponential falloff
+				Opacity = FMath::Clamp( Opacity, 0.0f, 1.0f );
+
+				const FColor Color = FLinearColor::Green.CopyWithNewOpacity( Opacity ).ToFColor( false );
+				AddMeshElementToOverlay( HoveredElementsComponent, FadingOutHoveredMeshElement, Color, HoveredSizeBias );
+			}
+		}
+	}
+
+	// Update selected mesh elements
+
+	const float SelectionAnimationDuration = MeshEd::SelectionAnimationDuration->GetFloat();
+	for( FMeshElement& SelectedMeshElement : SelectedMeshElements )
+	{
+		const float TimeSinceSelected = CurrentRealTime - SelectedMeshElement.LastSelectTime;
+		if( TimeSinceSelected < SelectionAnimationDuration )
+		{
+			bShouldUpdateSelectedElementsOverlay = true;
+			break;
+		}
+	}
+
+	if( bShouldUpdateSelectedElementsOverlay )
+	{
+		bShouldUpdateSelectedElementsOverlay = false;
+		UpdateSelectedElementsOverlay();
+	}
+
+	// Update debug normals/tangents
+	if( bShowVertexNormals )
+	{
+		UpdateDebugNormals();
+	}
+	else
+	{
+		DebugNormalsComponent->Clear();
+	}
+}
+
+
+void FMeshEditorMode::UpdateDebugNormals()
+{
+	// @todo mesheditor: There's nothing clever about this method.
+	// It just clears the old overlay lines and adds a bunch of new ones each tick.
+	// This should be a UWireframeMeshComponent with an adapter so that it can be updated incrementally as the mesh changes.
+	DebugNormalsComponent->Clear();
+
+	for( const FComponentAndEditableMesh& ComponentAndEditableMesh : SelectedComponentsAndEditableMeshes )
+	{
+		const UPrimitiveComponent* Component = ComponentAndEditableMesh.Component;
+		const UEditableMesh* EditableMesh = ComponentAndEditableMesh.EditableMesh;
+
+		const FMatrix ComponentToWorldMatrix = Component->GetRenderMatrix();
+
+		const uint32 PolygonArraySize = EditableMesh->GetPolygonArraySize();
+		for( uint32 PolygonIndex = 0; PolygonIndex < PolygonArraySize; ++PolygonIndex )
+		{
+			const FPolygonID PolygonID( PolygonIndex );
+			if( EditableMesh->IsValidPolygon( PolygonID ) )
+			{
+				// @todo mesheditor: total debug feature for now. Need a way of making this look nice.
+				const int32 PerimeterVertexCount = EditableMesh->GetPolygonPerimeterVertexCount( PolygonID );
+				const float Length = 10.0f; // @todo mesheditor: determine length of debug line from distance from the mesh origin to the camera?
+
+				for( int32 PerimeterVertexIndex = 0; PerimeterVertexIndex < PerimeterVertexCount; ++PerimeterVertexIndex )
+				{
+					const FVector Position( EditableMesh->GetPolygonPerimeterVertexAttribute( PolygonID, PerimeterVertexIndex, UEditableMeshAttribute::VertexPosition(), 0 ) );
+					const FVector Normal( EditableMesh->GetPolygonPerimeterVertexAttribute( PolygonID, PerimeterVertexIndex, UEditableMeshAttribute::VertexNormal(), 0 ) );
+					const FVector Tangent( EditableMesh->GetPolygonPerimeterVertexAttribute( PolygonID, PerimeterVertexIndex, UEditableMeshAttribute::VertexTangent(), 0 ) );
+
+					const FVector Start( ComponentToWorldMatrix.TransformPosition( Position ) );
+					const FVector NormalEnd( ComponentToWorldMatrix.TransformPosition( Position + Normal * Length ) );
+					const FVector TangentEnd( ComponentToWorldMatrix.TransformPosition( Position + Tangent * Length * 0.5f ) );
+
+					DebugNormalsComponent->AddLine( FOverlayLine( Start, NormalEnd, FColor::Magenta, 0.0f ) );
+					DebugNormalsComponent->AddLine( FOverlayLine( Start, TangentEnd, FColor::Yellow, 0.0f ) );
+				}
+			}
+		}
+	}
+}
+
+
+void FMeshEditorMode::RequestSelectedElementsOverlayUpdate()
+{
+	bShouldUpdateSelectedElementsOverlay = true;
+}
+
+
+void FMeshEditorMode::UpdateSelectedElementsOverlay()
+{
+	SelectedElementsComponent->Clear();
+	SelectedSubDElementsComponent->Clear();
+
+	const double CurrentRealTime = FSlateApplication::Get().GetCurrentTime();
+	const float SelectionAnimationDuration = MeshEd::SelectionAnimationDuration->GetFloat();
+
+	static TMap<TTuple<UPrimitiveComponent*, FEditableMeshSubMeshAddress>, TSet<FEdgeID>> SelectedEdgesByComponentsAndSubMeshes;
+	SelectedEdgesByComponentsAndSubMeshes.Reset();
+
+	for( const FMeshElement& SelectedMeshElement : SelectedMeshElements )
+	{
+		// Add selected elements to base cage overlay
+
+		const float TimeSinceSelected = CurrentRealTime - SelectedMeshElement.LastSelectTime;
+		const float SizeBias = MeshEd::SelectedSizeBias->GetFloat() + MeshEd::SelectedAnimationExtraSizeBias->GetFloat() * FMath::Clamp( 1.0f - ( TimeSinceSelected / SelectionAnimationDuration ), 0.0f, 1.0f );
+		const FColor Color = FLinearColor::White.ToFColor( false );
+
+		AddMeshElementToOverlay( SelectedElementsComponent, SelectedMeshElement, Color, SizeBias );
+
+		// If the editable mesh is previewing subdivisions, cache all selected edges (including sides of selected polygons)
+
+		if( SelectedMeshElement.Component.IsValid() )
+		{
+			UPrimitiveComponent* Component = SelectedMeshElement.Component.Get();
+			const FEditableMeshSubMeshAddress& SubMeshAddress = SelectedMeshElement.ElementAddress.SubMeshAddress;
+
+			const UEditableMesh* EditableMesh = FindEditableMesh( *Component, SubMeshAddress );
+			check( EditableMesh );
+			if( EditableMesh->IsPreviewingSubdivisions() )
+			{
+				TSet<FEdgeID>& EdgesToHighlight = SelectedEdgesByComponentsAndSubMeshes.FindOrAdd( MakeTuple( Component, SubMeshAddress ) );
+
+				if( SelectedMeshElement.ElementAddress.ElementType == EEditableMeshElementType::Edge )
+				{
+					const FEdgeID EdgeID( SelectedMeshElement.ElementAddress.ElementID );
+					EdgesToHighlight.Add( EdgeID );
+				}
+				else if( SelectedMeshElement.ElementAddress.ElementType == EEditableMeshElementType::Polygon )
+				{
+					const FPolygonID PolygonID( SelectedMeshElement.ElementAddress.ElementID );
+					const int32 PolygonEdgeCount = EditableMesh->GetPolygonPerimeterEdgeCount( PolygonID );
+					for( int32 EdgeIndex = 0; EdgeIndex < PolygonEdgeCount; ++EdgeIndex )
+					{
+						bool OutEdgeWindingReversed;
+						EdgesToHighlight.Add( EditableMesh->GetPolygonPerimeterEdge( PolygonID, EdgeIndex, OutEdgeWindingReversed ) );
+					}
+				}
+			}
+		}
+	}
+
+	// Add selected wires to subdivided mesh overlay
+
+	for( const auto& SelectedEdgesByComponentAndSubMesh : SelectedEdgesByComponentsAndSubMeshes )
+	{
+		const TTuple<UPrimitiveComponent*, FEditableMeshSubMeshAddress>& ComponentAndSubMesh = SelectedEdgesByComponentAndSubMesh.Key;
+		const TSet<FEdgeID>& EdgesToHighlight = SelectedEdgesByComponentAndSubMesh.Value;
+
+		if( EdgesToHighlight.Num() > 0 )
+		{
+			UPrimitiveComponent* Component = ComponentAndSubMesh.Get<0>();
+			const FEditableMeshSubMeshAddress& SubMeshAddress = ComponentAndSubMesh.Get<1>();
+
+			const UEditableMesh* EditableMesh = FindEditableMesh( *Component, SubMeshAddress );
+			check( EditableMesh != nullptr );
+
+			const FMatrix ComponentToWorldMatrix = Component->GetRenderMatrix();
+
+			const FSubdivisionLimitData& SubdivisionLimitData = EditableMesh->GetSubdivisionLimitData();
+
+			for( int32 WireEdgeNumber = 0; WireEdgeNumber < SubdivisionLimitData.SubdividedWireEdges.Num(); ++WireEdgeNumber )
+			{
+				const FSubdividedWireEdge& SubdividedWireEdge = SubdivisionLimitData.SubdividedWireEdges[ WireEdgeNumber ];
+
+				if( SubdividedWireEdge.CounterpartEdgeID != FEdgeID::Invalid && EdgesToHighlight.Contains( SubdividedWireEdge.CounterpartEdgeID ) )
+				{
+					const int32 EdgeVertexIndex0 = SubdividedWireEdge.EdgeVertex0PositionIndex;
+					const int32 EdgeVertexIndex1 = SubdividedWireEdge.EdgeVertex1PositionIndex;
+
+					const FVector Position0 = ComponentToWorldMatrix.TransformPosition( SubdivisionLimitData.VertexPositions[ EdgeVertexIndex0 ] );
+					const FVector Position1 = ComponentToWorldMatrix.TransformPosition( SubdivisionLimitData.VertexPositions[ EdgeVertexIndex1 ] );
+
+					const FColor Color = FLinearColor::White.CopyWithNewOpacity( 0.8f ).ToFColor( false );
+
+					SelectedSubDElementsComponent->AddLine( FOverlayLine( Position0, Position1, Color, 0.0f ) );
+				}
+			}
+		}
+	}
 }
 
 
@@ -1187,11 +1500,34 @@ void FMeshEditorMode::CommitEditableMeshIfNecessary( UEditableMesh* EditableMesh
 		NewEditableMesh->SetAllowSpatialDatabase(true);
 		NewEditableMesh->SetAllowCompact(true);
 
+		// Create a wireframe mesh for the base cage
+		UWireframeMesh* WireframeBaseCage = NewObject<UWireframeMesh>();
+
+		UMeshEditorStaticMeshAdapter* WireframeAdapter = NewObject<UMeshEditorStaticMeshAdapter>();
+		NewEditableMesh->Adapters.Add( WireframeAdapter );
+		WireframeAdapter->Initialize( NewEditableMesh, WireframeBaseCage );
+
+		// Create a wireframe mesh for the subdivided mesh
+		UWireframeMesh* WireframeSubdividedMesh = NewObject<UWireframeMesh>();
+
+		UMeshEditorSubdividedStaticMeshAdapter* WireframeSubdividedAdapter = NewObject<UMeshEditorSubdividedStaticMeshAdapter>();
+		NewEditableMesh->Adapters.Add( WireframeSubdividedAdapter );
+		WireframeSubdividedAdapter->Initialize( NewEditableMesh, WireframeSubdividedMesh );
+
+		// Rebuild mesh so that the wireframe meshes get their render data built through the adapters
+		NewEditableMesh->RebuildRenderMesh();
+
+		// Cache the editable mesh and the associated wireframe meshes
+		FEditableAndWireframeMeshes EditableAndWireframeMeshes;
+		EditableAndWireframeMeshes.EditableMesh = NewEditableMesh;
+		EditableAndWireframeMeshes.WireframeBaseCage = WireframeBaseCage;
+		EditableAndWireframeMeshes.WireframeSubdividedMesh = WireframeSubdividedMesh;
+
 		// Commit the editable mesh as a new instance in the static mesh component
 		const FEditableMeshSubMeshAddress& OldSubMeshAddress = EditableMesh->GetSubMeshAddress();
 		const FEditableMeshSubMeshAddress& NewSubMeshAddress = NewEditableMesh->GetSubMeshAddress();
 
-		CachedEditableMeshes.Add( NewSubMeshAddress, NewEditableMesh );
+		CachedEditableMeshes.Add( NewSubMeshAddress, EditableAndWireframeMeshes );
 
 		auto FixUpMeshElement = [ &OldSubMeshAddress, &NewSubMeshAddress, Component ]( FMeshElement& MeshElement )
 		{
@@ -2020,7 +2356,9 @@ void FMeshEditorMode::AddReferencedObjects( FReferenceCollector& Collector )
 
 	for( auto Pair : CachedEditableMeshes )
 	{
-		Collector.AddReferencedObject( Pair.Value );
+		Collector.AddReferencedObject( Pair.Value.EditableMesh );
+		Collector.AddReferencedObject( Pair.Value.WireframeBaseCage );
+		Collector.AddReferencedObject( Pair.Value.WireframeSubdividedMesh );
 	}
 
 	for( TTuple<UObject*, TUniquePtr<FChange>>& ObjectAndPreviewRevertChange : PreviewRevertChanges )
@@ -2036,286 +2374,68 @@ void FMeshEditorMode::AddReferencedObjects( FReferenceCollector& Collector )
 }
 
 
-void FMeshEditorMode::AddVertexToDynamicMesh( const UEditableMesh& EditableMesh, const FTransform& CameraToWorld, const bool bIsPerspectiveView, const FMatrix& ComponentToWorldMatrix, const FVertexID VertexID, const FColor ColorAndOpacity, const float SizeBias, const bool bApplyDepthBias, FDynamicMeshBuilder& MeshBuilder )
-{
-	const FVector VertexPosition = ComponentToWorldMatrix.TransformPosition( EditableMesh.GetVertexAttribute( VertexID, UEditableMeshAttribute::VertexPosition(), 0 ) );
-
-	const FVector CameraUpVector = CameraToWorld.TransformVector( FVector::UpVector ).GetSafeNormal();
-	const FVector CameraRightVector = CameraToWorld.TransformVector( FVector::RightVector ).GetSafeNormal();
-	const FVector CameraForwardVector = CameraToWorld.TransformVector( FVector::ForwardVector ).GetSafeNormal();
-
-	const FVector DirectionToCamera = bIsPerspectiveView ? ( CameraToWorld.GetLocation() - VertexPosition ).GetSafeNormal() : -CameraForwardVector;
-	const float DistanceToCamera = bIsPerspectiveView ? ( CameraToWorld.GetLocation() - VertexPosition ).Size() : 0.0f;
-	const float DistanceBias = bIsPerspectiveView ? MeshEd::OverlayPerspectiveDistanceBias->GetFloat() : MeshEd::OverlayOrthographicDistanceBias->GetFloat();
-	const float DistanceBasedScaling = DistanceBias + DistanceToCamera * MeshEd::OverlayDistanceScaleFactor->GetFloat();
-
-	const float SpriteSize = ( MeshEd::OverlayVertexSize->GetFloat() + SizeBias ) * DistanceBasedScaling;
-	const float HalfSpriteSize = SpriteSize * 0.5f;
-
-	// We're offsetting the geometry from the actual face a bit, to avoid z-fighting for this particular effect
-	const FVector VertexOffset = bApplyDepthBias ? ( DirectionToCamera * MeshEd::OverlayDepthOffset->GetFloat() * DistanceBasedScaling ) : FVector::ZeroVector;
-
-	FVector QuadPositions[ 4 ];
-	QuadPositions[ 0 ] = ( VertexPosition - CameraRightVector * HalfSpriteSize - CameraUpVector * HalfSpriteSize ) + VertexOffset;
-	QuadPositions[ 1 ] = ( VertexPosition - CameraRightVector * HalfSpriteSize + CameraUpVector * HalfSpriteSize ) + VertexOffset;
-	QuadPositions[ 2 ] = ( VertexPosition + CameraRightVector * HalfSpriteSize + CameraUpVector * HalfSpriteSize ) + VertexOffset;
-	QuadPositions[ 3 ] = ( VertexPosition + CameraRightVector * HalfSpriteSize - CameraUpVector * HalfSpriteSize ) + VertexOffset;
-
-	const int32 FirstVertexIndex = 
-		MeshBuilder.AddVertex( VertexOffset + QuadPositions[ 0 ], FVector2D( 0, 0 ), FVector( 1, 0, 0 ), FVector( 0, 1, 0 ), FVector( 0, 0, 1 ), ColorAndOpacity );
-	MeshBuilder.AddVertex( VertexOffset + QuadPositions[ 1 ], FVector2D( 0, 1 ), FVector( 1, 0, 0 ), FVector( 0, 1, 0 ), FVector( 0, 0, 1 ), ColorAndOpacity );
-	MeshBuilder.AddVertex( VertexOffset + QuadPositions[ 2 ], FVector2D( 1, 1 ), FVector( 1, 0, 0 ), FVector( 0, 1, 0 ), FVector( 0, 0, 1 ), ColorAndOpacity );
-	MeshBuilder.AddVertex( VertexOffset + QuadPositions[ 3 ], FVector2D( 1, 0 ), FVector( 1, 0, 0 ), FVector( 0, 1, 0 ), FVector( 0, 0, 1 ), ColorAndOpacity );
-
-	MeshBuilder.AddTriangle( FirstVertexIndex + 0, FirstVertexIndex + 1, FirstVertexIndex + 2 );
-	MeshBuilder.AddTriangle( FirstVertexIndex + 0, FirstVertexIndex + 2, FirstVertexIndex + 3 );
-}
-
-
-void FMeshEditorMode::AddThickLineToDynamicMesh( const FTransform& CameraToWorld, const bool bIsPerspectiveView, const FVector EdgeVertexPositions[2], const FColor ColorAndOpacity, const float SizeBias, const bool bApplyDepthBias, FDynamicMeshBuilder& MeshBuilder )
-{
-	const float DistanceBias = bIsPerspectiveView ? MeshEd::OverlayPerspectiveDistanceBias->GetFloat() : MeshEd::OverlayOrthographicDistanceBias->GetFloat();
-	const float Vertex0DistanceToCamera = bIsPerspectiveView ? ( CameraToWorld.GetLocation() - EdgeVertexPositions[ 0 ] ).Size() : 0.0f;
-	const float Vertex0DistanceBasedScaling = DistanceBias + Vertex0DistanceToCamera * MeshEd::OverlayDistanceScaleFactor->GetFloat();
-	const float Vertex1DistanceToCamera = bIsPerspectiveView ? ( CameraToWorld.GetLocation() - EdgeVertexPositions[ 1 ] ).Size() : 0.0f;
-	const float Vertex1DistanceBasedScaling = DistanceBias + Vertex1DistanceToCamera * MeshEd::OverlayDistanceScaleFactor->GetFloat();
-
-	const float Thickness = ( MeshEd::OverlayLineThickness->GetFloat() + SizeBias );
-	const float HalfThickness = Thickness * 0.5f;
-
-	const FVector CameraForwardVector = CameraToWorld.TransformVector( FVector::ForwardVector ).GetSafeNormal();
-
-	const FVector DirectionToCamera0 = bIsPerspectiveView ? ( CameraToWorld.GetLocation() - EdgeVertexPositions[ 0 ] ).GetSafeNormal() : -CameraForwardVector;
-	const FVector DirectionToCamera1 = bIsPerspectiveView ? ( CameraToWorld.GetLocation() - EdgeVertexPositions[ 1 ] ).GetSafeNormal() : -CameraForwardVector;
-	const FVector EdgeForward = ( EdgeVertexPositions[ 1 ] - EdgeVertexPositions[ 0 ] ).GetSafeNormal();
-	const FVector EdgeRight0 = FVector::CrossProduct( EdgeForward, DirectionToCamera0 ).GetSafeNormal();
-	const FVector EdgeRight1 = FVector::CrossProduct( EdgeForward, DirectionToCamera1 ).GetSafeNormal();
-
-	// We're offsetting the geometry from the actual face a bit, to avoid z-fighting for this particular effect
-	// @todo mesheditor urgent: Can we use a hardware depth bias instead?
-	const FVector Vertex0Offset = bApplyDepthBias ? ( DirectionToCamera0 * MeshEd::OverlayDepthOffset->GetFloat() * Vertex0DistanceBasedScaling ) : FVector::ZeroVector;
-	const FVector Vertex1Offset = bApplyDepthBias ? ( DirectionToCamera1 * MeshEd::OverlayDepthOffset->GetFloat() * Vertex1DistanceBasedScaling ) : FVector::ZeroVector;
-
-	FVector QuadPositions[ 4 ];
-	QuadPositions[ 0 ] = ( EdgeVertexPositions[ 0 ] - EdgeRight0 * HalfThickness * Vertex0DistanceBasedScaling ) + Vertex0Offset;
-	QuadPositions[ 1 ] = ( EdgeVertexPositions[ 0 ] + EdgeRight0 * HalfThickness * Vertex0DistanceBasedScaling ) + Vertex0Offset;
-	QuadPositions[ 2 ] = ( EdgeVertexPositions[ 1 ] + EdgeRight1 * HalfThickness * Vertex1DistanceBasedScaling ) + Vertex1Offset;
-	QuadPositions[ 3 ] = ( EdgeVertexPositions[ 1 ] - EdgeRight1 * HalfThickness * Vertex1DistanceBasedScaling ) + Vertex1Offset;
-
-	const int32 FirstVertexIndex = 
-		MeshBuilder.AddVertex( QuadPositions[ 0 ], FVector2D( 0, 0 ), FVector( 1, 0, 0 ), FVector( 0, 1, 0 ), FVector( 0, 0, 1 ), ColorAndOpacity );
-	MeshBuilder.AddVertex( QuadPositions[ 1 ], FVector2D( 0, 1 ), FVector( 1, 0, 0 ), FVector( 0, 1, 0 ), FVector( 0, 0, 1 ), ColorAndOpacity );
-	MeshBuilder.AddVertex( QuadPositions[ 2 ], FVector2D( 1, 1 ), FVector( 1, 0, 0 ), FVector( 0, 1, 0 ), FVector( 0, 0, 1 ), ColorAndOpacity );
-	MeshBuilder.AddVertex( QuadPositions[ 3 ], FVector2D( 1, 0 ), FVector( 1, 0, 0 ), FVector( 0, 1, 0 ), FVector( 0, 0, 1 ), ColorAndOpacity );
-
-	MeshBuilder.AddTriangle( FirstVertexIndex + 0, FirstVertexIndex + 1, FirstVertexIndex + 2 );
-	MeshBuilder.AddTriangle( FirstVertexIndex + 0, FirstVertexIndex + 2, FirstVertexIndex + 3 );
-}
-
-
-void FMeshEditorMode::AddEdgeToDynamicMesh( const UEditableMesh& EditableMesh, const FTransform& CameraToWorld, const bool bIsPerspectiveView, const FMatrix& ComponentToWorldMatrix, const FEdgeID EdgeID, const FColor ColorAndOpacity, const float SizeBias, FDynamicMeshBuilder& MeshBuilder )
+void FMeshEditorMode::AddMeshElementToOverlay( UOverlayComponent* OverlayComponent, const FMeshElement& MeshElement, const FColor Color, const float Size )
 {
-	FVertexID MeshVertexIDs[ 2 ];
-	EditableMesh.GetEdgeVertices( EdgeID, /* Out */ MeshVertexIDs[0], /* Out */ MeshVertexIDs[1] );
-
-	static TArray<FPolygonID> ConnectedPolygons;
-	EditableMesh.GetEdgeConnectedPolygons( EdgeID, /* Out */ ConnectedPolygons );
-
-	FVector EdgeVertexPositions[ 2 ];
-	for( uint32 EdgeVertexNumber = 0; EdgeVertexNumber < ARRAY_COUNT( MeshVertexIDs ); ++EdgeVertexNumber )
-	{
-		EdgeVertexPositions[ EdgeVertexNumber ] = ComponentToWorldMatrix.TransformPosition( EditableMesh.GetVertexAttribute( MeshVertexIDs[ EdgeVertexNumber ], UEditableMeshAttribute::VertexPosition(), 0 ) );
-	}
-
-	const bool bApplyDepthBias = true;
-	AddThickLineToDynamicMesh( CameraToWorld, bIsPerspectiveView, EdgeVertexPositions, ColorAndOpacity, SizeBias, bApplyDepthBias, MeshBuilder );
-}
-
-
-void FMeshEditorMode::AddPolygonToDynamicMesh( const UEditableMesh& EditableMesh, const FTransform& CameraToWorld, const bool bIsPerspectiveView, const FMatrix& ComponentToWorldMatrix, const FPolygonID PolygonID, const FColor ColorAndOpacity, const FColor HardEdgeColorAndOpacity, const float SizeBias, const bool bFillFaces, FDynamicMeshBuilder& VertexAndEdgeMeshBuilder, FDynamicMeshBuilder* PolygonFaceMeshBuilder )
-{
-	static TArray<FVertexID> MeshVertexIDs;
-	MeshVertexIDs.Reset();
-	EditableMesh.GetPolygonPerimeterVertices( PolygonID, /* Out */ MeshVertexIDs );
-
-	static TArray<FVector> PolygonPerimeterVertexPositions;
-	PolygonPerimeterVertexPositions.SetNumUninitialized( MeshVertexIDs.Num(), false );
-	for( uint32 PolygonVertexIndex = 0; PolygonVertexIndex < (uint32)MeshVertexIDs.Num(); ++PolygonVertexIndex )
-	{
-		PolygonPerimeterVertexPositions[ PolygonVertexIndex ] = ComponentToWorldMatrix.TransformPosition( EditableMesh.GetVertexAttribute( MeshVertexIDs[ PolygonVertexIndex ], UEditableMeshAttribute::VertexPosition(), 0 ) );
-	}
-
-	// @todo mesheditor: Support polygon hole contours!
-
-
-	const FVector PolygonNormal = ComponentToWorldMatrix.TransformVector( EditableMesh.ComputePolygonNormal( PolygonID ) ).GetSafeNormal();
-	const float Determinant = ComponentToWorldMatrix.Determinant();
-
-	const FVector CameraForwardVectorUnnormalized = CameraToWorld.TransformVector( FVector::ForwardVector );
-	const FVector FirstVertexDirectionToCameraUnnormalized = bIsPerspectiveView ? ( CameraToWorld.GetLocation() - PolygonPerimeterVertexPositions[0] ) : -CameraForwardVectorUnnormalized;
-
-	// @todo mesheditor: Backface culling in orthographic views may not always be desirable.  Make optional?
-	if( FVector::DotProduct( PolygonNormal, FirstVertexDirectionToCameraUnnormalized ) * Determinant < 0.0f )
-	{
-		// Ignore backfaced polys
-		return;
-	}
-
-	if( bFillFaces )
+	if( MeshElement.IsValidMeshElement() )
 	{
-		const float DistanceBias = bIsPerspectiveView ? MeshEd::OverlayPerspectiveDistanceBias->GetFloat() : MeshEd::OverlayOrthographicDistanceBias->GetFloat();
-
-		check( PolygonFaceMeshBuilder );
-		const uint32 PolygonTriangleCount = EditableMesh.GetPolygonTriangulatedTriangleCount( PolygonID );
-		for( uint32 PolygonTriangleNumber = 0; PolygonTriangleNumber < PolygonTriangleCount; ++PolygonTriangleNumber )
+		UEditableMesh* EditableMesh = FindOrCreateEditableMesh( *MeshElement.Component, MeshElement.ElementAddress.SubMeshAddress );
+		if( EditableMesh != nullptr )
 		{
-			FVector TriangleVertexPositions[ 3 ];
-			FVector VertexOffsets[ 3 ];
-			for( uint32 TriangleVertexNumber = 0; TriangleVertexNumber < 3; ++TriangleVertexNumber )
+			if( IsElementIDValid( MeshElement, EditableMesh ) )
 			{
-				TriangleVertexPositions[ TriangleVertexNumber ] =
-					ComponentToWorldMatrix.TransformPosition( 
-						EditableMesh.GetPolygonTriangulatedTriangleVertexPosition( PolygonID, PolygonTriangleNumber, TriangleVertexNumber ) );
-
-				// We're offsetting the geometry from the actual face a bit, to avoid z-fighting for this particular effect
-				const FVector DirectionToCamera = bIsPerspectiveView ? ( CameraToWorld.GetLocation() - TriangleVertexPositions[ TriangleVertexNumber ] ).GetSafeNormal() : -CameraForwardVectorUnnormalized.GetSafeNormal();
-				const float DistanceToCamera = bIsPerspectiveView ? ( CameraToWorld.GetLocation() - TriangleVertexPositions[ TriangleVertexNumber ] ).Size() : 0.0f;
-				const float DistanceBasedScaling = DistanceBias + DistanceToCamera * MeshEd::OverlayDistanceScaleFactor->GetFloat();
-				VertexOffsets[ TriangleVertexNumber ] = DirectionToCamera * MeshEd::OverlayDepthOffset->GetFloat() * DistanceBasedScaling;
-			}
-
-			// @todo mesheditor perf: Horribly inefficient way to draw a single triangle!  Also we should have a global utility function for this rather than using FDynamicMeshBuilder
-			// @todo mesheditor: Assumes vertices are in-order and form a single closed polygon.  Not accounting for holes.
+				UPrimitiveComponent* Component = MeshElement.Component.Get();
+				check( Component != nullptr );
 
-			const int32 FirstVertexIndex = 
-				PolygonFaceMeshBuilder->AddVertex( VertexOffsets[ 0 ] + TriangleVertexPositions[ 0 ], FVector2D( 0, 0 ), FVector( 1, 0, 0 ), FVector( 0, 1, 0 ), FVector( 0, 0, 1 ), ColorAndOpacity );
-			PolygonFaceMeshBuilder->AddVertex( VertexOffsets[ 1 ] + TriangleVertexPositions[ 1 ], FVector2D( 0, 1 ), FVector( 1, 0, 0 ), FVector( 0, 1, 0 ), FVector( 0, 0, 1 ), ColorAndOpacity );
-			PolygonFaceMeshBuilder->AddVertex( VertexOffsets[ 2 ] + TriangleVertexPositions[ 2 ], FVector2D( 1, 1 ), FVector( 1, 0, 0 ), FVector( 0, 1, 0 ), FVector( 0, 0, 1 ), ColorAndOpacity );
-
-			PolygonFaceMeshBuilder->AddTriangle( FirstVertexIndex + 0, FirstVertexIndex + 1, FirstVertexIndex + 2 );
-		}
-	}
-
-	{
-		static TArray<FEdgeID> PerimeterEdges;
-		EditableMesh.GetPolygonPerimeterEdges( PolygonID, /* Out */ PerimeterEdges );
-
-		// Draw polygon edges
-		// @todo mesheditor: Assumes no holes, etc.
-		for( const FEdgeID EdgeID : PerimeterEdges )
-		{
-			if( EditableMesh.GetEdgeAttribute( EdgeID, UEditableMeshAttribute::EdgeIsHard(), 0 ).X > 0.0f )
-			{
-				AddEdgeToDynamicMesh( EditableMesh, CameraToWorld, bIsPerspectiveView, ComponentToWorldMatrix, EdgeID, HardEdgeColorAndOpacity, SizeBias, VertexAndEdgeMeshBuilder );
-			}
-			else
-			{
-				AddEdgeToDynamicMesh( EditableMesh, CameraToWorld, bIsPerspectiveView, ComponentToWorldMatrix, EdgeID, ColorAndOpacity, SizeBias, VertexAndEdgeMeshBuilder );
-			}
-		}
-	}
-}
+				const FMatrix ComponentToWorldMatrix = Component->GetRenderMatrix();
 
+				switch( MeshElement.ElementAddress.ElementType )
+				{
+					case EEditableMeshElementType::Vertex:
+					{
+						const FVertexID VertexID( MeshElement.ElementAddress.ElementID );
+						const FVector Position( ComponentToWorldMatrix.TransformPosition( EditableMesh->GetVertexAttribute( VertexID, UEditableMeshAttribute::VertexPosition(), 0 ) ) );
 
-void FMeshEditorMode::DrawMeshElements( const FTransform& CameraToWorld, const bool bIsPerspectiveView, FViewport* Viewport, FPrimitiveDrawInterface* PDI, const TArrayView<FMeshElement>& MeshElements, const FColor Color, const bool bFillFaces, const float SizeBias, const TArray<FColor>* OptionalPerElementColors, const TArray<float>* OptionalPerElementSizeBiases )
-{
-	if( MeshElements.Num() > 0 )
-	{
-		// Must specify the same number of element colors and size biases as elements to draw
-		check( OptionalPerElementColors == nullptr || OptionalPerElementColors->Num() == MeshElements.Num() );
-		check( OptionalPerElementSizeBiases == nullptr || OptionalPerElementSizeBiases->Num() == MeshElements.Num() );
+						OverlayComponent->AddPoint( FOverlayPoint( Position, Color, Size ) );
+						break;
+					}
 
-		FDynamicMeshBuilder VertexAndEdgeMeshBuilder;
-		FDynamicMeshBuilder PolygonFaceMeshBuilder;
+					case EEditableMeshElementType::Edge:
+					{
+						const FEdgeID EdgeID( MeshElement.ElementAddress.ElementID );
+						const FVertexID StartVertexID( EditableMesh->GetEdgeVertex( EdgeID, 0 ) );
+						const FVertexID EndVertexID( EditableMesh->GetEdgeVertex( EdgeID, 1 ) );
+						const FVector StartPosition( ComponentToWorldMatrix.TransformPosition( EditableMesh->GetVertexAttribute( StartVertexID, UEditableMeshAttribute::VertexPosition(), 0 ) ) );
+						const FVector EndPosition( ComponentToWorldMatrix.TransformPosition( EditableMesh->GetVertexAttribute( EndVertexID, UEditableMeshAttribute::VertexPosition(), 0 ) ) );
 
-		uint32 NumPolygonsToDraw = 0;
-		
-		for( int32 MeshElementIndex = 0; MeshElementIndex < MeshElements.Num(); ++MeshElementIndex )
-		{
-			const FMeshElement& MeshElement = MeshElements[ MeshElementIndex ];
-			if( MeshElement.IsValidMeshElement() )
-			{
-				const FColor ElementColor = OptionalPerElementColors != nullptr ? ( *OptionalPerElementColors )[ MeshElementIndex ] : Color;
-				const float ElementSizeBias = OptionalPerElementSizeBiases != nullptr ? ( *OptionalPerElementSizeBiases )[ MeshElementIndex ] : SizeBias;
+						OverlayComponent->AddLine( FOverlayLine( StartPosition, EndPosition, Color, Size ) );
+						break;
+					}
 
-				UEditableMesh* EditableMesh = FindOrCreateEditableMesh( *MeshElement.Component, MeshElement.ElementAddress.SubMeshAddress );
-				if( EditableMesh != nullptr )
-				{
-					if( IsElementIDValid( MeshElement, EditableMesh ) )
+					case EEditableMeshElementType::Polygon:
 					{
-						UPrimitiveComponent* Component = MeshElement.Component.Get();
-						check( Component != nullptr )
-
-						const FMatrix ComponentToWorldMatrix = Component->GetRenderMatrix();
+						const FPolygonID PolygonID( MeshElement.ElementAddress.ElementID );
+						const int32 PolygonTriangleCount = EditableMesh->GetPolygonTriangulatedTriangleCount( PolygonID );
 
-						if( MeshElement.ElementAddress.ElementType == EEditableMeshElementType::Vertex )
+						for( int32 PolygonTriangle = 0; PolygonTriangle < PolygonTriangleCount; PolygonTriangle++ )
 						{
-							const bool bApplyDepthBias = true;
-							AddVertexToDynamicMesh(
-								*EditableMesh,
-								CameraToWorld,
-								bIsPerspectiveView,
-								ComponentToWorldMatrix,
-								FVertexID( MeshElement.ElementAddress.ElementID ),
-								ElementColor,
-								ElementSizeBias,
-								bApplyDepthBias,
-								VertexAndEdgeMeshBuilder );
-						}
-						else if( MeshElement.ElementAddress.ElementType == EEditableMeshElementType::Edge )
-						{
-							AddEdgeToDynamicMesh(
-								*EditableMesh,
-								CameraToWorld,
-								bIsPerspectiveView,
-								ComponentToWorldMatrix,
-								FEdgeID( MeshElement.ElementAddress.ElementID ),
-								ElementColor,
-								ElementSizeBias,
-								VertexAndEdgeMeshBuilder );
-						}
-						else if( MeshElement.ElementAddress.ElementType == EEditableMeshElementType::Polygon )
-						{
-							++NumPolygonsToDraw;
-
-							AddPolygonToDynamicMesh(
-								*EditableMesh,
-								CameraToWorld,
-								bIsPerspectiveView,
-								ComponentToWorldMatrix,
-								FPolygonID( MeshElement.ElementAddress.ElementID ),
-								ElementColor,
-								ElementColor,
-								ElementSizeBias,
-								bFillFaces,
-								VertexAndEdgeMeshBuilder,
-								&PolygonFaceMeshBuilder );
+							FVector VertexPositions[ 3 ];
+							for( int32 TriangleVertex = 0; TriangleVertex < 3; TriangleVertex++ )
+							{
+								VertexPositions[ TriangleVertex ] = ComponentToWorldMatrix.TransformPosition( EditableMesh->GetPolygonTriangulatedTriangleVertexPosition( PolygonID, PolygonTriangle, TriangleVertex ) );
+							}
+
+							OverlayComponent->AddTriangle( FOverlayTriangle( 
+								HoveredFaceMaterial,
+								FOverlayTriangleVertex( VertexPositions[ 0 ], FVector2D( 0, 0 ), FVector::UpVector, Color ),
+								FOverlayTriangleVertex( VertexPositions[ 1 ], FVector2D( 0, 1 ), FVector::UpVector, Color ),
+								FOverlayTriangleVertex( VertexPositions[ 2 ], FVector2D( 1, 1 ), FVector::UpVector, Color )
+							) );
 						}
+						break;
 					}
 				}
 			}
 		}
-
-		if( NumPolygonsToDraw && bFillFaces )
-		{
-			const bool bIsSelected = false;
-			FMaterialRenderProxy* MaterialRenderProxy = HoveredFaceMaterial->GetRenderProxy( bIsSelected );
-
-			const bool bDisableBackfaceCulling = true;	// @todo mesheditor: Disabling backface so we can see what's hovered even if its backfacing
-			const bool bReceivesDecals = false;
-			const FHitProxyId HitProxyId = FHitProxyId::InvisibleHitProxyId;	// Don't obscure selection!
-			PolygonFaceMeshBuilder.Draw( PDI, FMatrix::Identity, MaterialRenderProxy, SDPG_World, bDisableBackfaceCulling, bReceivesDecals, HitProxyId );
-		}
-
-		{
-			const bool bIsSelected = false;
-			FMaterialRenderProxy* MaterialRenderProxy = HoveredGeometryMaterial->GetRenderProxy( bIsSelected );
-
-			const bool bDisableBackfaceCulling = true;	// @todo mesheditor: Disabling backface so we can see what's hovered even if its backfacing
-			const bool bReceivesDecals = false;
-			const FHitProxyId HitProxyId = FHitProxyId::InvisibleHitProxyId;	// Don't obscure selection!
-			VertexAndEdgeMeshBuilder.Draw( PDI, FMatrix::Identity, MaterialRenderProxy, SDPG_World, bDisableBackfaceCulling, bReceivesDecals, HitProxyId );
-		}
 	}
 }
 
@@ -2324,17 +2444,6 @@ void FMeshEditorMode::Render( const FSceneView* SceneView, FViewport* Viewport,
 {
 	FEdMode::Render( SceneView, Viewport, PDI );
 
-	FTransform CameraToWorld;
-	{
-	    // The SceneView's view matrices are rotated 90 degrees in pitch, so we'll have to undo that to get
-		// back to Unreal standard coordinates
-	    const FMatrix& WorldToCameraMatrix = SceneView->ViewMatrices.GetViewMatrix();
-	    const FMatrix Rotate90( FPlane( 0, 0, 1, 0 ), FPlane( 1, 0, 0, 0 ), FPlane( 0, 1, 0, 0 ), FPlane( 0, 0, 0, 1 ) );
-		CameraToWorld.SetFromMatrix( ( WorldToCameraMatrix * Rotate90.Inverse() ).Inverse() );
-	}
-
-	const bool bIsPerspectiveView = SceneView->IsPerspectiveProjection();
-
 	// @todo mesheditor debug
 	if( MeshEd::ShowDebugStats->GetInt() > 0 && SelectedMeshElements.Num() > 0 )
 	{
@@ -2366,297 +2475,6 @@ void FMeshEditorMode::Render( const FSceneView* SceneView, FViewport* Viewport,
 			//}
 		}
 	}
-
-	// Draw all polygon edges for selected/hovered meshes (@todo mesheditor: This only draws polygons right now.  We should also draw "floating" edges and vertices too!)
-	if( MeshEd::ShowWiresForSelectedMeshes->GetInt() != 0 )
-	{
-		for( const FComponentAndEditableMesh& ComponentAndEditableMesh : SelectedComponentsAndEditableMeshes )
-		{
-			const UPrimitiveComponent* Component = ComponentAndEditableMesh.Component;
-			check( Component != nullptr );
-
-			const UEditableMesh* EditableMesh = ComponentAndEditableMesh.EditableMesh;
-			check( EditableMesh != nullptr );
-
-			const FMatrix ComponentToWorldMatrix = Component->GetRenderMatrix();
-
-			FDynamicMeshBuilder VertexAndEdgeMeshBuilder;
-
-			// Draw subdivision limit surface if subdivision preview is enabled
-			if( EditableMesh->IsPreviewingSubdivisions() )
-			{
-				// Figure out all of the edges that should appear 'selected' on the subdivision preview, taking into
-				// account all currently selected edges and polygons that share those edges
-				static TSet<FEdgeID> HighlightedEdgeIDs;
-				HighlightedEdgeIDs.Reset();
-
-				// @todo mesheditor perf: O(N) process to setup list of highlighted edges.  We could instead maintain this as selection changes.
-				for( FMeshElement& SelectedMeshElement : SelectedMeshElements )
-				{
-					if( SelectedMeshElement.IsValidMeshElement() &&
-						SelectedMeshElement.Component == Component &&
-						SelectedMeshElement.ElementAddress.SubMeshAddress == EditableMesh->GetSubMeshAddress() )
-					{
-						if( SelectedMeshElement.ElementAddress.ElementType == EEditableMeshElementType::Edge )
-						{
-							// OK, this edge is selected so it's definitely going to be highlighted
-							HighlightedEdgeIDs.Add( FEdgeID( SelectedMeshElement.ElementAddress.ElementID ) );
-						}
-						else if( SelectedMeshElement.ElementAddress.ElementType == EEditableMeshElementType::Polygon )
-						{
-							// Find all of this polygons edges.  They'll all need to be highlighted.
-							static TArray<FEdgeID> SelectedPolygonPerimeterEdgeIDs;
-							EditableMesh->GetPolygonPerimeterEdges(
-								FPolygonID( SelectedMeshElement.ElementAddress.ElementID ),
-								/* Out */ SelectedPolygonPerimeterEdgeIDs );
-
-							HighlightedEdgeIDs.Append( SelectedPolygonPerimeterEdgeIDs );
-						}
-					}
-				}
-
-
-				const float SizeBias = 0.0f;
-				const FColor SubdivisionEdgeColorAndOpacity = FLinearColor( 0.05f, 0.05f, 0.05f, 0.6f ).ToFColor( false );
-				const FColor BaseCageCounterpartEdgeColorAndOpacity = FLinearColor( 0.0f, 0.0f, 0.2f, 0.8f ).ToFColor( false );
-				const FColor SelectedBaseCageCounterpartEdgeColorAndOpacity = FLinearColor::White.CopyWithNewOpacity( 0.8f ).ToFColor( false );
-
-				const FSubdivisionLimitData& SubdivisionLimitData = EditableMesh->GetSubdivisionLimitData();
-
-				for( int32 WireEdgeNumber = 0; WireEdgeNumber < SubdivisionLimitData.SubdividedWireEdges.Num(); ++WireEdgeNumber )
-				{
-					const FSubdividedWireEdge& SubdividedWireEdge = SubdivisionLimitData.SubdividedWireEdges[ WireEdgeNumber ];
-
-					const int32 EdgeVertexIndexA = SubdividedWireEdge.EdgeVertex0PositionIndex;
-					const int32 EdgeVertexIndexB = SubdividedWireEdge.EdgeVertex1PositionIndex;
-
-					FVector EdgeVertexPositions[ 2 ];
-					EdgeVertexPositions[ 0 ] = SubdivisionLimitData.VertexPositions[ EdgeVertexIndexA ];
-					EdgeVertexPositions[ 1 ] = SubdivisionLimitData.VertexPositions[ EdgeVertexIndexB ];
-
-					FVector WorldSpaceEdgeVertexPositions[ 2 ];
-					for( uint32 EdgeVertexNumber = 0; EdgeVertexNumber < ARRAY_COUNT( EdgeVertexPositions ); ++EdgeVertexNumber )
-					{
-						WorldSpaceEdgeVertexPositions[ EdgeVertexNumber ] = ComponentToWorldMatrix.TransformPosition( EdgeVertexPositions[ EdgeVertexNumber ] );
-					}
-
-					FColor ColorAndOpacity = SubdivisionEdgeColorAndOpacity;
-					if( SubdividedWireEdge.CounterpartEdgeID != FEdgeID::Invalid )
-					{
-						if( HighlightedEdgeIDs.Contains( SubdividedWireEdge.CounterpartEdgeID ) )
-						{
-							ColorAndOpacity = SelectedBaseCageCounterpartEdgeColorAndOpacity;
-						}
-						else
-						{
-							ColorAndOpacity = BaseCageCounterpartEdgeColorAndOpacity;
-						}
-					}					
-
-					const bool bApplyDepthBias = true;
-					AddThickLineToDynamicMesh( CameraToWorld, bIsPerspectiveView, WorldSpaceEdgeVertexPositions, ColorAndOpacity, SizeBias, bApplyDepthBias, VertexAndEdgeMeshBuilder );
-				}
-			}
-
-			// Draw polygon mesh wires (or subdivision base cage, if previewing subdivisions)
-			{
-				const float Opacity = 0.85f;
-				const float SizeBias = 0.05f;
-				const bool bFillFaces = false;
-				const FColor ColorAndOpacity = FLinearColor( 0.0f, 0.0f, 0.3f ).CopyWithNewOpacity( Opacity ).ToFColor( false );
-				const FColor HardEdgeColorAndOpacity = FLinearColor( 0.0f, 0.3f, 0.3f ).CopyWithNewOpacity( Opacity ).ToFColor( false );
-
-				const uint32 PolygonArraySize = EditableMesh->GetPolygonArraySize();
-
-				for( uint32 PolygonIndex = 0; PolygonIndex < PolygonArraySize; ++PolygonIndex )
-				{
-					const FPolygonID PolygonID( PolygonIndex );
-					if( EditableMesh->IsValidPolygon( PolygonID ) )
-					{
-						AddPolygonToDynamicMesh(
-							*EditableMesh,
-							CameraToWorld,
-							bIsPerspectiveView,
-							ComponentToWorldMatrix,
-							PolygonID,
-							ColorAndOpacity,
-							HardEdgeColorAndOpacity,
-							SizeBias,
-							bFillFaces,
-							VertexAndEdgeMeshBuilder,
-							nullptr );
-
-						if( bShowVertexNormals )
-						{
-							// @todo mesheditor: total debug feature for now. Need a way of making this look nice.
-							const int32 PerimeterVertexCount = EditableMesh->GetPolygonPerimeterVertexCount( PolygonID );
-							const float Length = 10.0f; // @todo mesheditor: determine length of debug line from distance from the mesh origin to the camera?
-
-							for( int32 PerimeterVertexIndex = 0; PerimeterVertexIndex < PerimeterVertexCount; ++PerimeterVertexIndex )
-							{
-								FVector Position( EditableMesh->GetPolygonPerimeterVertexAttribute( PolygonID, PerimeterVertexIndex, UEditableMeshAttribute::VertexPosition(), 0 ) );
-								FVector Normal( EditableMesh->GetPolygonPerimeterVertexAttribute( PolygonID, PerimeterVertexIndex, UEditableMeshAttribute::VertexNormal(), 0 ) );
-								FVector Tangent( EditableMesh->GetPolygonPerimeterVertexAttribute( PolygonID, PerimeterVertexIndex, UEditableMeshAttribute::VertexTangent(), 0 ) );
-
-								FVector VertexNormalRenderPositions[] =
-								{
-									ComponentToWorldMatrix.TransformPosition( Position ),
-									ComponentToWorldMatrix.TransformPosition( Position + Normal * Length )
-								};
-
-								FVector VertexTangentRenderPositions[] =
-								{
-									ComponentToWorldMatrix.TransformPosition( Position ),
-									ComponentToWorldMatrix.TransformPosition( Position + Tangent * Length * 0.5f )
-								};
-
-								// We don't need any depth bias when drawing lines for face normals
-								const bool bApplyDepthBias = false;
-								AddThickLineToDynamicMesh(
-									CameraToWorld,
-									bIsPerspectiveView,
-									VertexNormalRenderPositions,
-									FColor::Magenta,
-									SizeBias,
-									bApplyDepthBias,
-									VertexAndEdgeMeshBuilder );
-
-								AddThickLineToDynamicMesh(
-									CameraToWorld,
-									bIsPerspectiveView,
-									VertexTangentRenderPositions,
-									FColor::Yellow,
-									SizeBias,
-									bApplyDepthBias,
-									VertexAndEdgeMeshBuilder );
-							}
-						}
-					}
-				}
-			}
-
-			const bool bIsSelected = false;
-			FMaterialRenderProxy* MaterialRenderProxy = HoveredGeometryMaterial->GetRenderProxy( bIsSelected );
-
-			const bool bDisableBackfaceCulling = true;	// @todo mesheditor: Disabling backface so we can see what's hovered even if its backfacing
-			const bool bReceivesDecals = false;
-			const FHitProxyId HitProxyId = FHitProxyId::InvisibleHitProxyId;	// Don't obscure selection!
-			VertexAndEdgeMeshBuilder.Draw( PDI, FMatrix::Identity, MaterialRenderProxy, SDPG_World, bDisableBackfaceCulling, bReceivesDecals, HitProxyId );
-		}
-	}
-
-
-	// Draw hovered elements
-	{
-		// Only draw hover if we're not in the middle of an interactive edit
-		if( ActiveAction == NAME_None )
-		{
-			const float HoveredSizeBias = MeshEd::HoveredSizeBias->GetFloat() + MeshEd::HoveredAnimationExtraSizeBias->GetFloat() * FMath::MakePulsatingValue( HoverFeedbackTimeValue, 0.5f );
-			{
-				static TArray<FMeshElement> HoveredMeshElementsToDraw;
-				HoveredMeshElementsToDraw.Reset();
-
-				// Draw hovered meshes
-				for( FMeshEditorInteractorData& MeshEditorInteractorData : MeshEditorInteractorDatas )
-				{
-					if( MeshElementSelectionMode == EEditableMeshElementType::Any ||
-					    MeshEditorInteractorData.HoveredMeshElement.ElementAddress.ElementType == MeshElementSelectionMode )
-					{
-						FMeshElement HoveredMeshElement = GetHoveredMeshElement( MeshEditorInteractorData.ViewportInteractor.Get() );
-						if( HoveredMeshElement.IsValidMeshElement() )
-						{
-							HoveredMeshElementsToDraw.Add( HoveredMeshElement );
-						}
-					}
-				}
-
-				const float Opacity = 1.0f;
-				const bool bFillFaces = true;
-				DrawMeshElements(
-					CameraToWorld,
-					bIsPerspectiveView,
-					Viewport,
-					PDI,
-					HoveredMeshElementsToDraw,
-					FLinearColor::Green.CopyWithNewOpacity( Opacity ).ToFColor( false ),
-					bFillFaces,
-					HoveredSizeBias );
-			}
-
-			// Draw meshes that were previously hovered
-			{
-				const double CurrentRealTime = FSlateApplication::Get().GetCurrentTime();
-
-				static TArray<FMeshElement> FadingOutHoveredMeshElementsToDraw;
-				FadingOutHoveredMeshElementsToDraw.Reset();
-				static TArray<FColor> PerElementColors;
-				PerElementColors.Reset();
-
-
-				const float HoverFadeTime = MeshEd::HoverFadeDuration->GetFloat();;
-				for( FMeshElement& FadingOutHoveredMeshElement : FadingOutHoveredMeshElements )
-				{
-					if( FadingOutHoveredMeshElement.IsValidMeshElement() )
-					{
-						const UEditableMesh* EditableMesh = FindEditableMesh( *FadingOutHoveredMeshElement.Component.Get(), FadingOutHoveredMeshElement.ElementAddress.SubMeshAddress );
-						if( IsElementIDValid( FadingOutHoveredMeshElement, EditableMesh ) )
-						{
-							const float TimeSinceLastHovered = CurrentRealTime - FadingOutHoveredMeshElement.LastHoverTime;
-							float Opacity = 1.0f - ( TimeSinceLastHovered / HoverFadeTime );
-							Opacity = Opacity * Opacity * Opacity * Opacity;		// Exponential falloff
-							Opacity = FMath::Clamp( Opacity, 0.0f, 1.0f );
-
-							FadingOutHoveredMeshElementsToDraw.Add( FadingOutHoveredMeshElement );
-							PerElementColors.Add( FLinearColor::Green.CopyWithNewOpacity( Opacity ).ToFColor( false ) );
-						}
-					}
-				}
-
-				const bool bFillFaces = true;
-				DrawMeshElements(
-					CameraToWorld,
-					bIsPerspectiveView,
-					Viewport,
-					PDI,
-					FadingOutHoveredMeshElementsToDraw,
-					FColor::White,	// Ignored, as we'll pass in per-element colors also
-					bFillFaces,
-					HoveredSizeBias,
-					&PerElementColors );
-			}
-		}
-
-		// Draw selected mesh elements
-		{
-			const double CurrentRealTime = FSlateApplication::Get().GetCurrentTime();
-
-			static TArray<float> PerElementSizeBiases;
-			PerElementSizeBiases.Reset();
-
-			const float SelectionAnimationDuration = MeshEd::SelectionAnimationDuration->GetFloat();
-			for( FMeshElement& SelectedMeshElement : SelectedMeshElements )
-			{
-				const float TimeSinceSelected = CurrentRealTime - SelectedMeshElement.LastSelectTime;
-
-				const float SizeBias = MeshEd::SelectedSizeBias->GetFloat() + MeshEd::SelectedAnimationExtraSizeBias->GetFloat() * FMath::Clamp( 1.0f - ( TimeSinceSelected / SelectionAnimationDuration ), 0.0f, 1.0f );
-				PerElementSizeBiases.Add( SizeBias );
-			}
-
-			const float Opacity = 1.0f;
-			const bool bFillFaces = true;
-			DrawMeshElements(
-				CameraToWorld,
-				bIsPerspectiveView,
-				Viewport,
-				PDI,
-				SelectedMeshElements,
-				FLinearColor::White.CopyWithNewOpacity( Opacity ).ToFColor( false ),
-				bFillFaces,
-				MeshEd::SelectedSizeBias->GetFloat(),
-				nullptr,
-				&PerElementSizeBiases );
-		}
-	}
 }
 
 
@@ -2951,8 +2769,9 @@ void FMeshEditorMode::OnViewportInteractionHoverUpdate( UViewportInteractor* Vie
 			}
 		}
 
-		// Are we hovering over something new? (or nothing?)  If so, then we'll fade out the old hovered mesh element
 		const FMeshElement& PreviouslyHoveredMeshElement = MeshEditorInteractorData.PreviouslyHoveredMeshElement;
+
+		// Are we hovering over something new? (or nothing?)  If so, then we'll fade out the old hovered mesh element
 		if( PreviouslyHoveredMeshElement.IsValidMeshElement() &&
 			!PreviouslyHoveredMeshElement.IsSameMeshElement( MeshEditorInteractorData.HoveredMeshElement ) )
 		{
@@ -3526,8 +3345,8 @@ void FMeshEditorMode::UpdateActiveAction( const bool bIsActionFinishing )
 			{
 				verify( !EditableMesh->AnyChangesToUndo() );
 
-				// @todo mesheditor perf: Polygons will actually be triangulated twice in the same frame when performing an operation that implicitly retriangulates at the end of the op (once before move, once after move.)  Not super ideal.
 				EditableMesh->MoveVertices( VerticesToMove );
+				RequestSelectedElementsOverlayUpdate();
 
 				TrackUndo( EditableMesh, EditableMesh->MakeUndo() );
 			}
@@ -4827,8 +4646,62 @@ void FMeshEditorMode::RefreshTransformables( const bool bNewObjectsSelected )
 }
 
 
+const FMeshEditorMode::FWireframeMeshComponents& FMeshEditorMode::CreateWireframeMeshComponents( UPrimitiveComponent* Component )
+{
+	FWireframeMeshComponents* WireframeMeshComponentsPtr = ComponentToWireframeComponentMap.Find( Component );
+	if( !WireframeMeshComponentsPtr )
+	{
+		WireframeMeshComponentsPtr = &ComponentToWireframeComponentMap.Add( Component );
+
+		const int32 LODIndex = 0;		// @todo mesheditor: We'll want to select an LOD to edit in various different wants (LOD that's visible, or manual user select, etc.)
+		const FEditableMeshSubMeshAddress SubMeshAddress = UEditableMeshFactory::MakeSubmeshAddress( Component, LODIndex );
+
+		const FEditableAndWireframeMeshes& EditableAndWireframeMeshes = CachedEditableMeshes.FindChecked( SubMeshAddress );
+
+		// Create the subdivided wireframe mesh component
+		UWireframeMeshComponent* WireframeSubdividedMeshComponent = NewObject<UWireframeMeshComponent>( WireframeComponentContainer );
+		WireframeSubdividedMeshComponent->SetMaterial( 0, SubdividedMeshWireMaterial );
+		WireframeSubdividedMeshComponent->TranslucencySortPriority = 100;
+		WireframeSubdividedMeshComponent->SetWireframeMesh( EditableAndWireframeMeshes.WireframeSubdividedMesh );
+		WireframeSubdividedMeshComponent->RegisterComponent();
+
+		// Create the base cage wireframe mesh component
+		UWireframeMeshComponent* WireframeMeshComponent = NewObject<UWireframeMeshComponent>( WireframeComponentContainer );
+		WireframeMeshComponent->SetMaterial( 0, WireMaterial );
+		WireframeMeshComponent->TranslucencySortPriority = 300;
+		WireframeMeshComponent->SetWireframeMesh( EditableAndWireframeMeshes.WireframeBaseCage );
+		WireframeMeshComponent->RegisterComponent();
+
+		WireframeMeshComponentsPtr->WireframeMeshComponent = WireframeMeshComponent;
+		WireframeMeshComponentsPtr->WireframeSubdividedMeshComponent = WireframeSubdividedMeshComponent;
+	}
+
+	return *WireframeMeshComponentsPtr;
+}
+
+
+void FMeshEditorMode::DestroyWireframeMeshComponents( UPrimitiveComponent* Component )
+{
+	FWireframeMeshComponents& WireframeMeshComponents = ComponentToWireframeComponentMap.FindChecked( Component );
+	WireframeMeshComponents.WireframeMeshComponent->DestroyComponent();
+	WireframeMeshComponents.WireframeSubdividedMeshComponent->DestroyComponent();
+	ComponentToWireframeComponentMap.Remove( Component );
+}
+
+
 void FMeshEditorMode::UpdateSelectedEditableMeshes()
 {
+	static TSet<UPrimitiveComponent*> DeselectedComponents;
+	DeselectedComponents.Reset();
+
+	// Make a list of components which have just been deselected.
+	// First add all the components which appear in the out-of-date list.
+	// Later, any components which are still selected will be removed from this list.
+	for( const FComponentAndEditableMesh& ComponentAndEditableMesh : SelectedComponentsAndEditableMeshes )
+	{
+		DeselectedComponents.Add( ComponentAndEditableMesh.Component );
+	}
+
 	SelectedEditableMeshes.Reset();
 	SelectedComponentsAndEditableMeshes.Reset();
 
@@ -4874,6 +4747,26 @@ void FMeshEditorMode::UpdateSelectedEditableMeshes()
 			}
 		}
 	}
+
+	for( const FComponentAndEditableMesh& ComponentAndEditableMesh : SelectedComponentsAndEditableMeshes )
+	{
+		if( DeselectedComponents.Contains( ComponentAndEditableMesh.Component ) )
+		{
+			DeselectedComponents.Remove( ComponentAndEditableMesh.Component );
+		}
+
+		const FWireframeMeshComponents& OverlayComponents = CreateWireframeMeshComponents( ComponentAndEditableMesh.Component );
+		const FTransform Transform = ComponentAndEditableMesh.Component->GetComponentTransform();
+		OverlayComponents.WireframeMeshComponent->SetWorldTransform( Transform );
+		OverlayComponents.WireframeSubdividedMeshComponent->SetWorldTransform( Transform );
+	}
+
+	for( UPrimitiveComponent* DeselectedComponent : DeselectedComponents )
+	{
+		DestroyWireframeMeshComponents( DeselectedComponent );
+	}
+
+	RequestSelectedElementsOverlayUpdate();
 }
 
 
diff --git a/Engine/Plugins/Editor/MeshEditor/Source/MeshEditor/MeshEditorMode.h b/Engine/Plugins/Editor/MeshEditor/Source/MeshEditor/MeshEditorMode.h
--- a/Engine/Plugins/Editor/MeshEditor/Source/MeshEditor/MeshEditorMode.h
+++ b/Engine/Plugins/Editor/MeshEditor/Source/MeshEditor/MeshEditorMode.h
@@ -16,6 +16,8 @@
 #include "Framework/Commands/UIAction.h"
 #include "Framework/Commands/UICommandInfo.h"
 #include "Framework/Commands/UICommandList.h"
+#include "WireframeMeshComponent.h"
+#include "OverlayComponent.h"
 #include "MeshEditorMode.generated.h"
 
 
@@ -139,6 +141,7 @@ public:
 protected:
 
 	// FEdMode interface
+	virtual void Initialize() override;
 	virtual void Enter() override;
 	virtual void Exit() override;
 	virtual void Tick( FEditorViewportClient* ViewportClient, float DeltaTime ) override;
@@ -212,21 +215,6 @@ protected:
 	/** Gets the container of all the assets used in the mesh editor */
 	const class UMeshEditorAssetContainer& GetAssetContainer() const;
 
-	/** Fills the specified dynamic mesh builder with primitives to render a mesh vertex */
-	void AddVertexToDynamicMesh( const UEditableMesh& EditableMesh, const FTransform& CameraToWorld, const bool bIsPerspectiveView, const FMatrix& ComponentToWorldMatrix, const FVertexID VertexID, const FColor ColorAndOpacity, const float SizeBias, const bool bApplyDepthBias, class FDynamicMeshBuilder& MeshBuilder );
-		
-	/** Fills the specified dynamic mesh builder with primitives to a thick line.  Incoming positions are in world space. */
-	void AddThickLineToDynamicMesh( const FTransform& CameraToWorld, const bool bIsPerspectiveView, const FVector EdgeVertexPositions[2], const FColor ColorAndOpacity, const float SizeBias, const bool bApplyDepthBias, FDynamicMeshBuilder& MeshBuilder );
-
-	/** Fills the specified dynamic mesh builder with primitives to render a mesh edge */
-	void AddEdgeToDynamicMesh( const UEditableMesh& EditableMesh, const FTransform& CameraToWorld, const bool bIsPerspectiveView, const FMatrix& ComponentToWorldMatrix, const FEdgeID EdgeID, const FColor ColorAndOpacity, const float SizeBias, class FDynamicMeshBuilder& MeshBuilder );
-
-	/** Fills the specified dynamic mesh builders with primitives to render a polygon and it's edges */
-	void AddPolygonToDynamicMesh( const UEditableMesh& EditableMesh, const FTransform& CameraToWorld, const bool bIsPerspectiveView, const FMatrix& ComponentToWorldMatrix, const FPolygonID PolygonID, const FColor ColorAndOpacity, const FColor HardEdgeColorAndOpacity, const float SizeBias, const bool bFillFaces, class FDynamicMeshBuilder& VertexAndEdgeMeshBuilder, class FDynamicMeshBuilder* PolygonFaceMeshBuilder );
-
-	/** Renders the spceified mesh element */
-	void DrawMeshElements( const FTransform& CameraToWorld, const bool bIsPerspectiveView, FViewport* Viewport, FPrimitiveDrawInterface* PDI, const TArrayView<FMeshElement>& MeshElements, const FColor Color, const bool bFillFaces, const float HoverAnimation, const TArray<FColor>* OptionalPerElementColors = nullptr, const TArray<float>* OptionalPerElementSizeBiases = nullptr );
-
 	/** Called every frame for each viewport interactor to update what's under the cursor */
 	void OnViewportInteractionHoverUpdate( class UViewportInteractor* ViewportInteractor, FVector& OutHoverImpactPoint, bool& bWasHandled );
 		
@@ -363,6 +351,30 @@ protected:
 	/** Plays sound when mesh edit action was finished */
 	void PlayFinishActionSound( FName NewAction, UViewportInteractor* ActionInteractor = nullptr);
 
+	struct FWireframeMeshComponents
+	{
+		UWireframeMeshComponent* WireframeMeshComponent;
+		UWireframeMeshComponent* WireframeSubdividedMeshComponent;
+	};
+
+	/** Creates wireframe mesh components for an editable mesh */
+	const FWireframeMeshComponents& CreateWireframeMeshComponents( UPrimitiveComponent* Component );
+
+	/** Destroys overlay components for an editable mesh */
+	void DestroyWireframeMeshComponents( UPrimitiveComponent* Component );
+
+	/** Adds a mesh element to the given overlay component */
+	void AddMeshElementToOverlay( UOverlayComponent* OverlayComponent, const FMeshElement& MeshElement, const FColor Color, const float Size );
+
+	/** Request an update to the selected elements overlay */
+	void RequestSelectedElementsOverlayUpdate();
+
+	/** Updates the selected elements overlay */
+	void UpdateSelectedElementsOverlay();
+
+	/** Updates debug normals/tangents */
+	void UpdateDebugNormals();
+
 
 protected:
 
@@ -496,6 +508,18 @@ protected:
 	/** Material to use to render hovered triangles or faces */
 	UMaterialInterface* HoveredFaceMaterial;
 
+	/** Material to use for wireframe base cage render */
+	UMaterialInterface* WireMaterial;
+
+	/** Material to use for wireframe subdivided mesh render */
+	UMaterialInterface* SubdividedMeshWireMaterial;
+
+	/** Material to use for the overlay lines */
+	UMaterialInterface* OverlayLineMaterial;
+
+	/** Material to use for the overlay points */
+	UMaterialInterface* OverlayPointMaterial;
+
 	/** Hover feedback animation time value, ever-incrementing until selection changes */
 	double HoverFeedbackTimeValue;
 
@@ -543,11 +567,22 @@ protected:
 	TArray<FMeshElement> FadingOutHoveredMeshElements;
 
 
-
+	/** Data pertaining to an editable mesh:
+	    The editable mesh itself, plus two wireframe visualizations (base cage and subD), which
+		are updated as the editable mesh is updated. */
+	struct FEditableAndWireframeMeshes
+	{
+		UEditableMesh* EditableMesh;
+		UWireframeMesh* WireframeBaseCage;
+		UWireframeMesh* WireframeSubdividedMesh;
+	};
 
 	/** Cached editable meshes */
 	// @todo mesheditor: Need to expire these at some point, otherwise we just grow and grow */
-	TMap< FEditableMeshSubMeshAddress, class UEditableMesh* > CachedEditableMeshes;
+	TMap<FEditableMeshSubMeshAddress, FEditableAndWireframeMeshes> CachedEditableMeshes;
+
+	/** Map mesh components to their wireframe overlays */
+	TMap<UPrimitiveComponent*, FWireframeMeshComponents> ComponentToWireframeComponentMap;
 
 	/** Manages saving undo for selected mesh elements while we're dragging them around */
 	FTrackingTransaction TrackingTransaction;
@@ -582,6 +617,21 @@ protected:
 	    such as the selection set.  We need this because the UE4 undo system requires a UObject, but we're an FEdMode. */
 	UMeshEditorModeProxyObject* MeshEditorModeProxyObject;
 
+	/** Actor which holds wireframe mesh components */
+	AActor* WireframeComponentContainer;
+
+	/** Component containing hovered elements */
+	UOverlayComponent* HoveredElementsComponent;
+
+	/** Component containing selected elements */
+	UOverlayComponent* SelectedElementsComponent;
+
+	/** Component containing subD mesh selected wires */
+	UOverlayComponent* SelectedSubDElementsComponent;
+
+	/** Component containing debug normals/tangents */
+	UOverlayComponent* DebugNormalsComponent;
+
 	/** When performing an interactive action that was initiated using an interactor, this is the interactor that was used. */
 	class UViewportInteractor* ActiveActionInteractor;
 
@@ -651,6 +701,9 @@ protected:
 	/** Requested focus to selection */
 	bool bShouldFocusToSelection;
 
+	/** Requested selected elements overlay update */
+	bool bShouldUpdateSelectedElementsOverlay;
+
 	/** Whether edits are made per-instance or not */
 	bool bPerInstanceEdits;
 
diff --git a/Engine/Plugins/Editor/MeshEditor/Source/MeshEditor/MeshEditorStaticMeshAdapter.cpp b/Engine/Plugins/Editor/MeshEditor/Source/MeshEditor/MeshEditorStaticMeshAdapter.cpp
new file mode 100644
--- /dev/null
+++ b/Engine/Plugins/Editor/MeshEditor/Source/MeshEditor/MeshEditorStaticMeshAdapter.cpp
@@ -0,0 +1,317 @@
+// Copyright 1998-2015 Epic Games, Inc. All Rights Reserved.
+
+#include "MeshEditorStaticMeshAdapter.h"
+#include "EditableMesh.h"
+#include "WireframeMeshComponent.h"
+
+
+static FColor GetEdgeColor( bool bIsHardEdge )
+{
+	return bIsHardEdge ? FColor( 0, 76, 76 ) : FColor( 0, 0, 76 );
+}
+
+
+UMeshEditorStaticMeshAdapter::UMeshEditorStaticMeshAdapter()
+{
+}
+
+
+void UMeshEditorStaticMeshAdapter::Initialize( UEditableMesh* EditableMesh, UWireframeMesh* InWireframeMesh )
+{
+	WireframeMesh = InWireframeMesh;
+}
+
+
+void UMeshEditorStaticMeshAdapter::OnRebuildRenderMeshStart( const UEditableMesh* EditableMesh, const bool bInvalidateLighting )
+{
+	for( TObjectIterator<UWireframeMeshComponent> It( RF_ClassDefaultObject, false, EInternalObjectFlags::PendingKill ); It; ++It )
+	{
+		if( It->GetWireframeMesh() == WireframeMesh )
+		{
+			if( It->IsRenderStateCreated() )
+			{
+				check( It->IsRegistered() );
+				It->UnregisterComponent();
+			}
+		}
+	}
+
+	WireframeMesh->ReleaseResources();
+
+	// Flush the rendering commands generated by the detachments.
+	FlushRenderingCommands();
+}
+
+
+void UMeshEditorStaticMeshAdapter::OnRebuildRenderMesh( const UEditableMesh* EditableMesh )
+{
+	WireframeMesh->Reset();
+
+	// Add all vertices
+	for( TSparseArray<FMeshVertex>::TConstIterator It( EditableMesh->Vertices ); It; ++It )
+	{
+		const FVertexID VertexID( It.GetIndex() );
+		WireframeMesh->AddVertex( VertexID );
+		WireframeMesh->SetVertexPosition( VertexID, EditableMesh->GetVertexAttribute( VertexID, UEditableMeshAttribute::VertexPosition(), 0 ) );
+	}
+
+	// Add all edges
+	for( TSparseArray<FMeshEdge>::TConstIterator It( EditableMesh->Edges ); It; ++It )
+	{
+		const FEdgeID EdgeID( It.GetIndex() );
+		WireframeMesh->AddEdge( EdgeID );
+		WireframeMesh->SetEdgeVertices( EdgeID, It->VertexIDs[ 0 ], It->VertexIDs[ 1 ] );
+		WireframeMesh->SetEdgeColor( EdgeID, GetEdgeColor( EditableMesh->GetEdgeAttribute( EdgeID, UEditableMeshAttribute::EdgeIsHard(), 0 ).X > 0.0f ) );
+	}
+
+	// Add all polygons and edge instances
+	// @todo mesheditor wireframe: holes not yet handled
+	for( TSparseArray<FMeshPolygon>::TConstIterator It( EditableMesh->Polygons ); It; ++It )
+	{
+		const FPolygonID PolygonID( It.GetIndex() );
+
+		WireframeMesh->AddPolygon( PolygonID );
+		WireframeMesh->SetPolygonNormal( PolygonID, EditableMesh->ComputePolygonNormal( PolygonID ) );
+
+		const int32 NumEdges = EditableMesh->GetPolygonPerimeterEdgeCount( PolygonID );
+		for( int32 Index = 0; Index < NumEdges; ++Index )
+		{
+			bool bOutWindingIsReversed;
+			WireframeMesh->AddEdgeInstance( EditableMesh->GetPolygonPerimeterEdge( PolygonID, Index, bOutWindingIsReversed ), PolygonID );
+		}
+	}
+}
+
+
+void UMeshEditorStaticMeshAdapter::OnRebuildRenderMeshFinish( const UEditableMesh* EditableMesh, const bool bRebuildBoundsAndCollision, const bool bIsPreviewRollback )
+{
+	for( const FPolygonID PolygonID : EditableMesh->PolygonsPendingNewTangentBasis )
+	{
+		const FMeshPolygon& Polygon = EditableMesh->Polygons[ PolygonID.GetValue() ];
+		WireframeMesh->SetPolygonNormal( PolygonID, Polygon.PolygonNormal );
+	}
+
+	WireframeMesh->InitResources();
+
+	for( TObjectIterator<UWireframeMeshComponent> It( RF_ClassDefaultObject, false, EInternalObjectFlags::PendingKill ); It; ++It )
+	{
+		if( It->GetWireframeMesh() == WireframeMesh )
+		{
+			if( !It->IsRenderStateCreated() )
+			{
+				It->RegisterComponent();
+			}
+		}
+	}
+}
+
+
+void UMeshEditorStaticMeshAdapter::OnStartModification( const UEditableMesh* EditableMesh, const EMeshModificationType MeshModificationType, const EMeshTopologyChange MeshTopologyChange )
+{
+}
+
+
+void UMeshEditorStaticMeshAdapter::OnEndModification( const UEditableMesh* EditableMesh )
+{
+}
+
+
+void UMeshEditorStaticMeshAdapter::OnReindexElements( const UEditableMesh* EditableMesh, const FElementIDRemappings& Remappings )
+{
+}
+
+
+bool UMeshEditorStaticMeshAdapter::IsCommitted( const UEditableMesh* EditableMesh ) const
+{
+	return false;
+}
+
+
+bool UMeshEditorStaticMeshAdapter::IsCommittedAsInstance( const UEditableMesh* EditableMesh ) const
+{
+	return false;
+}
+
+
+void UMeshEditorStaticMeshAdapter::OnCommit( UEditableMesh* EditableMesh )
+{
+}
+
+
+UEditableMesh* UMeshEditorStaticMeshAdapter::OnCommitInstance( UEditableMesh* EditableMesh, UPrimitiveComponent* ComponentToInstanceTo )
+{
+	return nullptr;
+}
+
+
+void UMeshEditorStaticMeshAdapter::OnRevert( UEditableMesh* EditableMesh )
+{
+}
+
+
+UEditableMesh* UMeshEditorStaticMeshAdapter::OnRevertInstance( UEditableMesh* EditableMesh )
+{
+	return nullptr;
+}
+
+
+void UMeshEditorStaticMeshAdapter::OnPropagateInstanceChanges( UEditableMesh* EditableMesh )
+{
+}
+
+
+void UMeshEditorStaticMeshAdapter::OnDeleteVertexInstances( const UEditableMesh* EditableMesh, const TArray<FVertexInstanceID>& VertexInstanceIDs )
+{
+}
+
+
+void UMeshEditorStaticMeshAdapter::OnDeleteOrphanVertices( const UEditableMesh* EditableMesh, const TArray<FVertexID>& VertexIDs )
+{
+	for( const FVertexID VertexID : VertexIDs )
+	{
+		WireframeMesh->RemoveVertex( VertexID );
+	}
+}
+
+
+void UMeshEditorStaticMeshAdapter::OnCreateEmptyVertexRange( const UEditableMesh* EditableMesh, const TArray<FVertexID>& VertexIDs )
+{
+	for( const FVertexID VertexID : VertexIDs )
+	{
+		WireframeMesh->AddVertex( VertexID );
+	}
+}
+
+
+void UMeshEditorStaticMeshAdapter::OnCreateVertices( const UEditableMesh* EditableMesh, const TArray<FVertexID>& VertexIDs )
+{
+	for( const FVertexID VertexID : VertexIDs )
+	{
+		WireframeMesh->AddVertex( VertexID );
+		WireframeMesh->SetVertexPosition( VertexID, EditableMesh->GetVertexAttribute( VertexID, UEditableMeshAttribute::VertexPosition(), 0 ) );
+	}
+}
+
+
+void UMeshEditorStaticMeshAdapter::OnCreateVertexInstances( const UEditableMesh* EditableMesh, const TArray<FVertexInstanceID>& VertexInstanceIDs )
+{
+}
+
+
+void UMeshEditorStaticMeshAdapter::OnSetVertexAttribute( const UEditableMesh* EditableMesh, const FVertexID VertexID, const FName AttributeName, const int32 AttributeIndex, const FVector4 AttributeValue )
+{
+	if( AttributeName == UEditableMeshAttribute::VertexPosition() )
+	{
+		WireframeMesh->SetVertexPosition( VertexID, AttributeValue );
+	}
+}
+
+
+void UMeshEditorStaticMeshAdapter::OnSetVertexInstanceAttribute( const UEditableMesh* EditableMesh, const FVertexInstanceID VertexInstanceID, const FName AttributeName, const int32 AttributeIndex, const FVector4 AttributeValue )
+{
+}
+
+
+void UMeshEditorStaticMeshAdapter::OnCreateEdges( const UEditableMesh* EditableMesh, const TArray<FEdgeID>& EdgeIDs )
+{
+	for( const FEdgeID EdgeID : EdgeIDs )
+	{
+		const FMeshEdge& Edge = EditableMesh->Edges[ EdgeID.GetValue() ];
+
+		WireframeMesh->AddEdge( EdgeID );
+		WireframeMesh->SetEdgeVertices( EdgeID, Edge.VertexIDs[ 0 ], Edge.VertexIDs[ 1 ] );
+		WireframeMesh->SetEdgeColor( EdgeID, GetEdgeColor( EditableMesh->GetEdgeAttribute( EdgeID, UEditableMeshAttribute::EdgeIsHard(), 0 ).X > 0.0f ) );
+
+		for( const FPolygonID PolygonID : Edge.ConnectedPolygons )
+		{
+			WireframeMesh->AddEdgeInstance( EdgeID, PolygonID );
+		}
+	}
+}
+
+
+void UMeshEditorStaticMeshAdapter::OnDeleteEdges( const UEditableMesh* EditableMesh, const TArray<FEdgeID>& EdgeIDs )
+{
+	for( const FEdgeID EdgeID : EdgeIDs )
+	{
+		const FMeshEdge& Edge = EditableMesh->Edges[ EdgeID.GetValue() ];
+
+		for( const FPolygonID PolygonID : Edge.ConnectedPolygons )
+		{
+			WireframeMesh->RemoveEdgeInstance( EdgeID, PolygonID );
+		}
+
+		WireframeMesh->RemoveEdge( EdgeID );
+	}
+}
+
+
+void UMeshEditorStaticMeshAdapter::OnSetEdgesVertices( const UEditableMesh* EditableMesh, const TArray<FEdgeID>& EdgeIDs )
+{
+	for( const FEdgeID EdgeID : EdgeIDs )
+	{
+		WireframeMesh->SetEdgeVertices( EdgeID, EditableMesh->GetEdgeVertex( EdgeID, 0 ), EditableMesh->GetEdgeVertex( EdgeID, 1 ) );
+	}
+}
+
+
+void UMeshEditorStaticMeshAdapter::OnSetEdgeAttribute( const UEditableMesh* EditableMesh, const FEdgeID EdgeID, const FName AttributeName, const int32 AttributeIndex, const FVector4 AttributeValue )
+{
+	if( AttributeName == UEditableMeshAttribute::EdgeIsHard() )
+	{
+		WireframeMesh->SetEdgeColor( EdgeID, GetEdgeColor( AttributeValue.X > 0.0f ) );
+	}
+}
+
+
+void UMeshEditorStaticMeshAdapter::OnCreatePolygons( const UEditableMesh* EditableMesh, const TArray<FPolygonID>& PolygonIDs )
+{
+	for( const FPolygonID PolygonID : PolygonIDs )
+	{
+		WireframeMesh->AddPolygon( PolygonID );
+		WireframeMesh->SetPolygonNormal( PolygonID, EditableMesh->ComputePolygonNormal( PolygonID ) );
+
+		const int32 NumEdges = EditableMesh->GetPolygonPerimeterEdgeCount( PolygonID );
+		for( int32 Index = 0; Index < NumEdges; ++Index )
+		{
+			bool bOutWindingIsReversed;
+			WireframeMesh->AddEdgeInstance( EditableMesh->GetPolygonPerimeterEdge( PolygonID, Index, bOutWindingIsReversed ), PolygonID );
+		}
+	}
+}
+
+
+void UMeshEditorStaticMeshAdapter::OnDeletePolygons( const UEditableMesh* EditableMesh, const TArray<FPolygonID>& PolygonIDs )
+{
+	for( const FPolygonID PolygonID : PolygonIDs )
+	{
+		const int32 NumEdges = EditableMesh->GetPolygonPerimeterEdgeCount( PolygonID );
+		for( int32 Index = 0; Index < NumEdges; ++Index )
+		{
+			bool bOutWindingIsReversed;
+			WireframeMesh->RemoveEdgeInstance( EditableMesh->GetPolygonPerimeterEdge( PolygonID, Index, bOutWindingIsReversed ), PolygonID );
+		}
+
+		WireframeMesh->RemovePolygon( PolygonID );
+	}
+}
+
+
+void UMeshEditorStaticMeshAdapter::OnChangePolygonVertexInstances( const UEditableMesh* EditableMesh, const TArray<FPolygonID>& PolygonIDs )
+{
+}
+
+
+void UMeshEditorStaticMeshAdapter::OnCreatePolygonGroups( const UEditableMesh* EditableMesh, const TArray<FPolygonGroupID>& PolygonGroupIDs )
+{
+}
+
+
+void UMeshEditorStaticMeshAdapter::OnDeletePolygonGroups( const UEditableMesh* EditableMesh, const TArray<FPolygonGroupID>& PolygonGroupIDs )
+{
+}
+
+
+void UMeshEditorStaticMeshAdapter::OnRetriangulatePolygons( const UEditableMesh* EditableMesh, const TArray<FPolygonID>& PolygonIDs )
+{
+}
diff --git a/Engine/Plugins/Editor/MeshEditor/Source/MeshEditor/MeshEditorStaticMeshAdapter.h b/Engine/Plugins/Editor/MeshEditor/Source/MeshEditor/MeshEditorStaticMeshAdapter.h
new file mode 100644
--- /dev/null
+++ b/Engine/Plugins/Editor/MeshEditor/Source/MeshEditor/MeshEditorStaticMeshAdapter.h
@@ -0,0 +1,66 @@
+// Copyright 1998-2017 Epic Games, Inc. All Rights Reserved.
+
+#pragma once
+
+#include "CoreMinimal.h"
+#include "EditableMeshTypes.h"
+#include "EditableMeshAdapter.h"
+#include "StaticMeshResources.h"
+#include "MeshEditorStaticMeshAdapter.generated.h"
+
+
+UCLASS(MinimalAPI)
+class UMeshEditorStaticMeshAdapter : public UEditableMeshAdapter
+{
+	GENERATED_BODY()
+
+public:
+
+	/** Default constructor that initializes good defaults for UEditableStaticMeshAdapter */
+	UMeshEditorStaticMeshAdapter();
+
+	/** Creates a editable static mesh from the specified component and sub-mesh address */
+	void Initialize( UEditableMesh* EditableMesh, class UWireframeMesh* WireframeMesh );
+
+	virtual void OnRebuildRenderMeshStart( const UEditableMesh* EditableMesh, const bool bInvalidateLighting ) override;
+	virtual void OnRebuildRenderMesh( const UEditableMesh* EditableMesh ) override;
+	virtual void OnRebuildRenderMeshFinish( const UEditableMesh* EditableMesh, const bool bRebuildBoundsAndCollision, const bool bIsPreviewRollback ) override;
+	virtual void OnStartModification( const UEditableMesh* EditableMesh, const EMeshModificationType MeshModificationType, const EMeshTopologyChange MeshTopologyChange ) override;
+	virtual void OnEndModification( const UEditableMesh* EditableMesh ) override;
+	virtual void OnReindexElements( const UEditableMesh* EditableMesh, const FElementIDRemappings& Remappings ) override;
+	virtual bool IsCommitted( const UEditableMesh* EditableMesh ) const override;
+	virtual bool IsCommittedAsInstance( const UEditableMesh* EditableMesh ) const override;
+	virtual void OnCommit( UEditableMesh* EditableMesh ) override;
+	virtual UEditableMesh* OnCommitInstance( UEditableMesh* EditableMesh, UPrimitiveComponent* ComponentToInstanceTo ) override;
+	virtual void OnRevert( UEditableMesh* EditableMesh ) override;
+	virtual UEditableMesh* OnRevertInstance( UEditableMesh* EditableMesh ) override;
+	virtual void OnPropagateInstanceChanges( UEditableMesh* EditableMesh ) override;
+
+	virtual void OnDeleteVertexInstances( const UEditableMesh* EditableMesh, const TArray<FVertexInstanceID>& VertexInstanceIDs ) override;
+	virtual void OnDeleteOrphanVertices( const UEditableMesh* EditableMesh, const TArray<FVertexID>& VertexIDs ) override;
+	virtual void OnCreateEmptyVertexRange( const UEditableMesh* EditableMesh, const TArray<FVertexID>& VertexIDs ) override;
+	virtual void OnCreateVertices( const UEditableMesh* EditableMesh, const TArray<FVertexID>& VertexIDs ) override;
+	virtual void OnCreateVertexInstances( const UEditableMesh* EditableMesh, const TArray<FVertexInstanceID>& VertexInstanceIDs ) override;
+	virtual void OnSetVertexAttribute( const UEditableMesh* EditableMesh, const FVertexID VertexID, const FName AttributeName, const int32 AttributeIndex, const FVector4 AttributeValue ) override;
+	virtual void OnSetVertexInstanceAttribute( const UEditableMesh* EditableMesh, const FVertexInstanceID VertexInstanceID, const FName AttributeName, const int32 AttributeIndex, const FVector4 AttributeValue ) override;
+	virtual void OnCreateEdges( const UEditableMesh* EditableMesh, const TArray<FEdgeID>& EdgeIDs ) override;
+	virtual void OnDeleteEdges( const UEditableMesh* EditableMesh, const TArray<FEdgeID>& EdgeIDs ) override;
+	virtual void OnSetEdgesVertices( const UEditableMesh* EditableMesh, const TArray<FEdgeID>& EdgeIDs ) override;
+	virtual void OnSetEdgeAttribute( const UEditableMesh* EditableMesh, const FEdgeID EdgeID, const FName AttributeName, const int32 AttributeIndex, const FVector4 AttributeValue ) override;
+	virtual void OnCreatePolygons( const UEditableMesh* EditableMesh, const TArray<FPolygonID>& PolygonIDs ) override;
+	virtual void OnDeletePolygons( const UEditableMesh* EditableMesh, const TArray<FPolygonID>& PolygonIDs ) override;
+	virtual void OnChangePolygonVertexInstances( const UEditableMesh* EditableMesh, const TArray<FPolygonID>& PolygonIDs ) override;
+	virtual void OnCreatePolygonGroups( const UEditableMesh* EditableMesh, const TArray<FPolygonGroupID>& PolygonGroupIDs ) override;
+	virtual void OnDeletePolygonGroups( const UEditableMesh* EditableMesh, const TArray<FPolygonGroupID>& PolygonGroupIDs ) override;
+	virtual void OnRetriangulatePolygons( const UEditableMesh* EditableMesh, const TArray<FPolygonID>& PolygonIDs ) override;
+
+
+private:
+
+	/** The wireframe mesh asset we're representing */
+	UPROPERTY()
+	UWireframeMesh* WireframeMesh;
+
+	UPROPERTY()
+	int32 StaticMeshLODIndex;
+};
diff --git a/Engine/Plugins/Editor/MeshEditor/Source/MeshEditor/MeshEditorSubdividedStaticMeshAdapter.cpp b/Engine/Plugins/Editor/MeshEditor/Source/MeshEditor/MeshEditorSubdividedStaticMeshAdapter.cpp
new file mode 100644
--- /dev/null
+++ b/Engine/Plugins/Editor/MeshEditor/Source/MeshEditor/MeshEditorSubdividedStaticMeshAdapter.cpp
@@ -0,0 +1,236 @@
+// Copyright 1998-2015 Epic Games, Inc. All Rights Reserved.
+
+#include "MeshEditorSubdividedStaticMeshAdapter.h"
+#include "EditableMesh.h"
+#include "WireframeMeshComponent.h"
+
+
+static FColor GetEdgeColor( bool bIsSubdividedWireEdge )
+{
+	return bIsSubdividedWireEdge ? FColor( 0, 0, 51, 204 ) : FColor( 13, 13, 13, 153 );
+}
+
+
+UMeshEditorSubdividedStaticMeshAdapter::UMeshEditorSubdividedStaticMeshAdapter()
+{
+}
+
+
+void UMeshEditorSubdividedStaticMeshAdapter::Initialize( UEditableMesh* EditableMesh, UWireframeMesh* InWireframeMesh )
+{
+	WireframeMesh = InWireframeMesh;
+}
+
+
+void UMeshEditorSubdividedStaticMeshAdapter::OnRebuildRenderMeshStart( const UEditableMesh* EditableMesh, const bool bInvalidateLighting )
+{
+	for( TObjectIterator<UWireframeMeshComponent> It( RF_ClassDefaultObject, false, EInternalObjectFlags::PendingKill ); It; ++It )
+	{
+		if( It->GetWireframeMesh() == WireframeMesh )
+		{
+			if( It->IsRenderStateCreated() )
+			{
+				check( It->IsRegistered() );
+				It->UnregisterComponent();
+			}
+		}
+	}
+
+	WireframeMesh->ReleaseResources();
+
+	// Flush the rendering commands generated by the detachments.
+	FlushRenderingCommands();
+}
+
+
+void UMeshEditorSubdividedStaticMeshAdapter::OnRebuildRenderMesh( const UEditableMesh* EditableMesh )
+{
+	const FSubdivisionLimitData& SubdivisionLimitData = EditableMesh->GetSubdivisionLimitData();
+
+	WireframeMesh->Reset();
+
+	if( EditableMesh->IsPreviewingSubdivisions() )
+	{
+		int32 VertexIndex = 0;
+		for( const FVector& VertexPosition : SubdivisionLimitData.VertexPositions )
+		{
+			const FVertexID VertexID( VertexIndex );
+			WireframeMesh->AddVertex( VertexID );
+			WireframeMesh->SetVertexPosition( VertexID, VertexPosition );
+			VertexIndex++;
+		}
+
+		// Not going to perform backface culling (we would need precalculated edge normals to do that, and ideally
+		// two edge instances - one per adjacent polygon), so just use a zero vector for the polygon normal.
+		const FPolygonID DummyPolygonID( 0 );
+		WireframeMesh->AddPolygon( DummyPolygonID );
+		WireframeMesh->SetPolygonNormal( DummyPolygonID, FVector::ZeroVector );
+
+		int32 EdgeIndex = 0;
+		for( const FSubdividedWireEdge& SubdividedWireEdge : SubdivisionLimitData.SubdividedWireEdges )
+		{
+			const FEdgeID EdgeID( EdgeIndex );
+			WireframeMesh->AddEdge( EdgeID );
+			WireframeMesh->SetEdgeVertices( EdgeID, FVertexID( SubdividedWireEdge.EdgeVertex0PositionIndex ), FVertexID( SubdividedWireEdge.EdgeVertex1PositionIndex ) );
+			WireframeMesh->SetEdgeColor( EdgeID, GetEdgeColor( SubdividedWireEdge.CounterpartEdgeID != FEdgeID::Invalid ) );
+			WireframeMesh->AddEdgeInstance( EdgeID, DummyPolygonID );
+			EdgeIndex++;
+		}
+	}
+}
+
+
+void UMeshEditorSubdividedStaticMeshAdapter::OnRebuildRenderMeshFinish( const UEditableMesh* EditableMesh, const bool bRebuildBoundsAndCollision, const bool bIsPreviewRollback )
+{
+	WireframeMesh->InitResources();
+
+	for( TObjectIterator<UWireframeMeshComponent> It( RF_ClassDefaultObject, false, EInternalObjectFlags::PendingKill ); It; ++It )
+	{
+		if( It->GetWireframeMesh() == WireframeMesh )
+		{
+			if( !It->IsRenderStateCreated() )
+			{
+				It->RegisterComponent();
+			}
+		}
+	}
+}
+
+
+void UMeshEditorSubdividedStaticMeshAdapter::OnStartModification( const UEditableMesh* EditableMesh, const EMeshModificationType MeshModificationType, const EMeshTopologyChange MeshTopologyChange )
+{
+}
+
+
+void UMeshEditorSubdividedStaticMeshAdapter::OnEndModification( const UEditableMesh* EditableMesh )
+{
+}
+
+
+void UMeshEditorSubdividedStaticMeshAdapter::OnReindexElements( const UEditableMesh* EditableMesh, const FElementIDRemappings& Remappings )
+{
+}
+
+
+bool UMeshEditorSubdividedStaticMeshAdapter::IsCommitted( const UEditableMesh* EditableMesh ) const
+{
+	return false;
+}
+
+
+bool UMeshEditorSubdividedStaticMeshAdapter::IsCommittedAsInstance( const UEditableMesh* EditableMesh ) const
+{
+	return false;
+}
+
+
+void UMeshEditorSubdividedStaticMeshAdapter::OnCommit( UEditableMesh* EditableMesh )
+{
+}
+
+
+UEditableMesh* UMeshEditorSubdividedStaticMeshAdapter::OnCommitInstance( UEditableMesh* EditableMesh, UPrimitiveComponent* ComponentToInstanceTo )
+{
+	return nullptr;
+}
+
+
+void UMeshEditorSubdividedStaticMeshAdapter::OnRevert( UEditableMesh* EditableMesh )
+{
+}
+
+
+UEditableMesh* UMeshEditorSubdividedStaticMeshAdapter::OnRevertInstance( UEditableMesh* EditableMesh )
+{
+	return nullptr;
+}
+
+
+void UMeshEditorSubdividedStaticMeshAdapter::OnPropagateInstanceChanges( UEditableMesh* EditableMesh )
+{
+}
+
+
+void UMeshEditorSubdividedStaticMeshAdapter::OnDeleteVertexInstances( const UEditableMesh* EditableMesh, const TArray<FVertexInstanceID>& VertexInstanceIDs )
+{
+}
+
+
+void UMeshEditorSubdividedStaticMeshAdapter::OnDeleteOrphanVertices( const UEditableMesh* EditableMesh, const TArray<FVertexID>& VertexIDs )
+{
+}
+
+
+void UMeshEditorSubdividedStaticMeshAdapter::OnCreateEmptyVertexRange( const UEditableMesh* EditableMesh, const TArray<FVertexID>& VertexIDs )
+{
+}
+
+
+void UMeshEditorSubdividedStaticMeshAdapter::OnCreateVertices( const UEditableMesh* EditableMesh, const TArray<FVertexID>& VertexIDs )
+{
+}
+
+
+void UMeshEditorSubdividedStaticMeshAdapter::OnCreateVertexInstances( const UEditableMesh* EditableMesh, const TArray<FVertexInstanceID>& VertexInstanceIDs )
+{
+}
+
+
+void UMeshEditorSubdividedStaticMeshAdapter::OnSetVertexAttribute( const UEditableMesh* EditableMesh, const FVertexID VertexID, const FName AttributeName, const int32 AttributeIndex, const FVector4 AttributeValue )
+{
+}
+
+
+void UMeshEditorSubdividedStaticMeshAdapter::OnSetVertexInstanceAttribute( const UEditableMesh* EditableMesh, const FVertexInstanceID VertexInstanceID, const FName AttributeName, const int32 AttributeIndex, const FVector4 AttributeValue )
+{
+}
+
+
+void UMeshEditorSubdividedStaticMeshAdapter::OnCreateEdges( const UEditableMesh* EditableMesh, const TArray<FEdgeID>& EdgeIDs )
+{
+}
+
+
+void UMeshEditorSubdividedStaticMeshAdapter::OnDeleteEdges( const UEditableMesh* EditableMesh, const TArray<FEdgeID>& EdgeIDs )
+{
+}
+
+
+void UMeshEditorSubdividedStaticMeshAdapter::OnSetEdgesVertices( const UEditableMesh* EditableMesh, const TArray<FEdgeID>& EdgeIDs )
+{
+}
+
+
+void UMeshEditorSubdividedStaticMeshAdapter::OnSetEdgeAttribute( const UEditableMesh* EditableMesh, const FEdgeID EdgeID, const FName AttributeName, const int32 AttributeIndex, const FVector4 AttributeValue )
+{
+}
+
+
+void UMeshEditorSubdividedStaticMeshAdapter::OnCreatePolygons( const UEditableMesh* EditableMesh, const TArray<FPolygonID>& PolygonIDs )
+{
+}
+
+
+void UMeshEditorSubdividedStaticMeshAdapter::OnDeletePolygons( const UEditableMesh* EditableMesh, const TArray<FPolygonID>& PolygonIDs )
+{
+}
+
+
+void UMeshEditorSubdividedStaticMeshAdapter::OnChangePolygonVertexInstances( const UEditableMesh* EditableMesh, const TArray<FPolygonID>& PolygonIDs )
+{
+}
+
+
+void UMeshEditorSubdividedStaticMeshAdapter::OnCreatePolygonGroups( const UEditableMesh* EditableMesh, const TArray<FPolygonGroupID>& PolygonGroupIDs )
+{
+}
+
+
+void UMeshEditorSubdividedStaticMeshAdapter::OnDeletePolygonGroups( const UEditableMesh* EditableMesh, const TArray<FPolygonGroupID>& PolygonGroupIDs )
+{
+}
+
+
+void UMeshEditorSubdividedStaticMeshAdapter::OnRetriangulatePolygons( const UEditableMesh* EditableMesh, const TArray<FPolygonID>& PolygonIDs )
+{
+}
diff --git a/Engine/Plugins/Editor/MeshEditor/Source/MeshEditor/MeshEditorSubdividedStaticMeshAdapter.h b/Engine/Plugins/Editor/MeshEditor/Source/MeshEditor/MeshEditorSubdividedStaticMeshAdapter.h
new file mode 100644
--- /dev/null
+++ b/Engine/Plugins/Editor/MeshEditor/Source/MeshEditor/MeshEditorSubdividedStaticMeshAdapter.h
@@ -0,0 +1,66 @@
+// Copyright 1998-2017 Epic Games, Inc. All Rights Reserved.
+
+#pragma once
+
+#include "CoreMinimal.h"
+#include "EditableMeshTypes.h"
+#include "EditableMeshAdapter.h"
+#include "StaticMeshResources.h"
+#include "MeshEditorSubdividedStaticMeshAdapter.generated.h"
+
+
+UCLASS(MinimalAPI)
+class UMeshEditorSubdividedStaticMeshAdapter : public UEditableMeshAdapter
+{
+	GENERATED_BODY()
+
+public:
+
+	/** Default constructor that initializes good defaults for UMeshEditorSubdividedStaticMeshAdapter */
+	UMeshEditorSubdividedStaticMeshAdapter();
+
+	/** Creates a editable static mesh from the specified component and sub-mesh address */
+	void Initialize( UEditableMesh* EditableMesh, class UWireframeMesh* WireframeMesh );
+
+	virtual void OnRebuildRenderMeshStart( const UEditableMesh* EditableMesh, const bool bInvalidateLighting ) override;
+	virtual void OnRebuildRenderMesh( const UEditableMesh* EditableMesh ) override;
+	virtual void OnRebuildRenderMeshFinish( const UEditableMesh* EditableMesh, const bool bRebuildBoundsAndCollision, const bool bIsPreviewRollback ) override;
+	virtual void OnStartModification( const UEditableMesh* EditableMesh, const EMeshModificationType MeshModificationType, const EMeshTopologyChange MeshTopologyChange ) override;
+	virtual void OnEndModification( const UEditableMesh* EditableMesh ) override;
+	virtual void OnReindexElements( const UEditableMesh* EditableMesh, const FElementIDRemappings& Remappings ) override;
+	virtual bool IsCommitted( const UEditableMesh* EditableMesh ) const override;
+	virtual bool IsCommittedAsInstance( const UEditableMesh* EditableMesh ) const override;
+	virtual void OnCommit( UEditableMesh* EditableMesh ) override;
+	virtual UEditableMesh* OnCommitInstance( UEditableMesh* EditableMesh, UPrimitiveComponent* ComponentToInstanceTo ) override;
+	virtual void OnRevert( UEditableMesh* EditableMesh ) override;
+	virtual UEditableMesh* OnRevertInstance( UEditableMesh* EditableMesh ) override;
+	virtual void OnPropagateInstanceChanges( UEditableMesh* EditableMesh ) override;
+
+	virtual void OnDeleteVertexInstances( const UEditableMesh* EditableMesh, const TArray<FVertexInstanceID>& VertexInstanceIDs ) override;
+	virtual void OnDeleteOrphanVertices( const UEditableMesh* EditableMesh, const TArray<FVertexID>& VertexIDs ) override;
+	virtual void OnCreateEmptyVertexRange( const UEditableMesh* EditableMesh, const TArray<FVertexID>& VertexIDs ) override;
+	virtual void OnCreateVertices( const UEditableMesh* EditableMesh, const TArray<FVertexID>& VertexIDs ) override;
+	virtual void OnCreateVertexInstances( const UEditableMesh* EditableMesh, const TArray<FVertexInstanceID>& VertexInstanceIDs ) override;
+	virtual void OnSetVertexAttribute( const UEditableMesh* EditableMesh, const FVertexID VertexID, const FName AttributeName, const int32 AttributeIndex, const FVector4 AttributeValue ) override;
+	virtual void OnSetVertexInstanceAttribute( const UEditableMesh* EditableMesh, const FVertexInstanceID VertexInstanceID, const FName AttributeName, const int32 AttributeIndex, const FVector4 AttributeValue ) override;
+	virtual void OnCreateEdges( const UEditableMesh* EditableMesh, const TArray<FEdgeID>& EdgeIDs ) override;
+	virtual void OnDeleteEdges( const UEditableMesh* EditableMesh, const TArray<FEdgeID>& EdgeIDs ) override;
+	virtual void OnSetEdgesVertices( const UEditableMesh* EditableMesh, const TArray<FEdgeID>& EdgeIDs ) override;
+	virtual void OnSetEdgeAttribute( const UEditableMesh* EditableMesh, const FEdgeID EdgeID, const FName AttributeName, const int32 AttributeIndex, const FVector4 AttributeValue ) override;
+	virtual void OnCreatePolygons( const UEditableMesh* EditableMesh, const TArray<FPolygonID>& PolygonIDs ) override;
+	virtual void OnDeletePolygons( const UEditableMesh* EditableMesh, const TArray<FPolygonID>& PolygonIDs ) override;
+	virtual void OnChangePolygonVertexInstances( const UEditableMesh* EditableMesh, const TArray<FPolygonID>& PolygonIDs ) override;
+	virtual void OnCreatePolygonGroups( const UEditableMesh* EditableMesh, const TArray<FPolygonGroupID>& PolygonGroupIDs ) override;
+	virtual void OnDeletePolygonGroups( const UEditableMesh* EditableMesh, const TArray<FPolygonGroupID>& PolygonGroupIDs ) override;
+	virtual void OnRetriangulatePolygons( const UEditableMesh* EditableMesh, const TArray<FPolygonID>& PolygonIDs ) override;
+
+
+private:
+
+	/** The wireframe mesh asset we're representing */
+	UPROPERTY()
+	UWireframeMesh* WireframeMesh;
+
+	UPROPERTY()
+	int32 StaticMeshLODIndex;
+};
diff --git a/Engine/Plugins/Editor/MeshEditor/Source/MeshEditor/OverlayComponent.cpp b/Engine/Plugins/Editor/MeshEditor/Source/MeshEditor/OverlayComponent.cpp
new file mode 100644
--- /dev/null
+++ b/Engine/Plugins/Editor/MeshEditor/Source/MeshEditor/OverlayComponent.cpp
@@ -0,0 +1,574 @@
+// Copyright 1998-2017 Epic Games, Inc. All Rights Reserved. 
+
+#include "OverlayComponent.h"
+#include "RenderingThread.h"
+#include "PrimitiveViewRelevance.h"
+#include "PrimitiveSceneProxy.h"
+#include "VertexFactory.h"
+#include "MaterialShared.h"
+#include "Engine/CollisionProfile.h"
+#include "Materials/Material.h"
+#include "LocalVertexFactory.h"
+#include "SceneManagement.h"
+#include "DynamicMeshBuilder.h"
+#include "EngineGlobals.h"
+#include "Engine/Engine.h"
+
+
+struct FOverlayVertex
+{
+	FOverlayVertex() {}
+
+	FOverlayVertex( const FVector& InPosition, const FVector2D& InUV, const FVector& InTangentX, const FVector& InTangentZ, const FColor& InColor )
+		: Position( InPosition ),
+		  UV( InUV ),
+		  TangentX( InTangentX ),
+		  TangentZ( InTangentZ ),
+		  Color( InColor )
+	{
+		TangentZ.Vector.W = 255;
+	}
+
+	FVector Position;
+	FVector2D UV;
+	FPackedNormal TangentX;
+	FPackedNormal TangentZ;
+	FColor Color;
+};
+
+
+class FOverlayVertexBuffer : public FVertexBuffer 
+{
+public:
+	TArray<FOverlayVertex> Vertices;
+
+	virtual void InitRHI() override
+	{
+		FRHIResourceCreateInfo CreateInfo;
+		void* VertexBufferData = nullptr;
+		VertexBufferRHI = RHICreateAndLockVertexBuffer( Vertices.Num() * sizeof( FOverlayVertex ), BUF_Static, CreateInfo, VertexBufferData );
+
+		// Copy the vertex data into the vertex buffer.		
+		FMemory::Memcpy( VertexBufferData, Vertices.GetData(), Vertices.Num() * sizeof( FOverlayVertex ) );
+		RHIUnlockVertexBuffer( VertexBufferRHI );
+	}
+
+};
+
+
+class FOverlayIndexBuffer : public FIndexBuffer 
+{
+public:
+	TArray<int32> Indices;
+
+	virtual void InitRHI() override
+	{
+		FRHIResourceCreateInfo CreateInfo;
+		void* Buffer = nullptr;
+		IndexBufferRHI = RHICreateAndLockIndexBuffer( sizeof( int32 ), Indices.Num() * sizeof( int32 ), BUF_Static, CreateInfo, Buffer );
+
+		// Write the indices to the index buffer.		
+		FMemory::Memcpy( Buffer, Indices.GetData(), Indices.Num() * sizeof( int32 ) );
+		RHIUnlockIndexBuffer( IndexBufferRHI );
+	}
+};
+
+
+class FOverlayVertexFactory : public FLocalVertexFactory
+{
+public:
+
+	FOverlayVertexFactory() = default;
+
+	/** Init function that should only be called on render thread. */
+	void Init_RenderThread( const FOverlayVertexBuffer* VertexBuffer )
+	{
+		check( IsInRenderingThread() );
+
+		FDataType NewData;
+		NewData.PositionComponent = STRUCTMEMBER_VERTEXSTREAMCOMPONENT( VertexBuffer, FOverlayVertex, Position, VET_Float3 );
+		NewData.TextureCoordinates.Add(
+			FVertexStreamComponent( VertexBuffer, STRUCT_OFFSET( FOverlayVertex, UV ), sizeof( FOverlayVertex ), VET_Float2 )
+			);
+		NewData.TangentBasisComponents[ 0 ] = STRUCTMEMBER_VERTEXSTREAMCOMPONENT( VertexBuffer, FOverlayVertex, TangentX, VET_PackedNormal );
+		NewData.TangentBasisComponents[ 1 ] = STRUCTMEMBER_VERTEXSTREAMCOMPONENT( VertexBuffer, FOverlayVertex, TangentZ, VET_PackedNormal );
+		NewData.ColorComponent = STRUCTMEMBER_VERTEXSTREAMCOMPONENT( VertexBuffer, FOverlayVertex, Color, VET_Color );
+
+		SetData( NewData );
+	}
+
+	/** Initialization */
+	void Init( const FOverlayVertexBuffer* VertexBuffer )
+	{
+		if ( IsInRenderingThread() )
+		{
+			Init_RenderThread( VertexBuffer );
+		}
+		else
+		{
+			ENQUEUE_UNIQUE_RENDER_COMMAND_TWOPARAMETER(
+				InitOverlayVertexFactory,
+				FOverlayVertexFactory*, VertexFactory, this,
+				const FOverlayVertexBuffer*, VertexBuffer, VertexBuffer,
+				{
+					VertexFactory->Init_RenderThread( VertexBuffer );
+				}
+			);
+		}	
+	}
+};
+
+
+struct FMeshBatchData
+{
+	FMaterialRenderProxy* MaterialProxy;
+	FOverlayVertexBuffer VertexBuffer;
+	FOverlayIndexBuffer IndexBuffer;
+	FOverlayVertexFactory VertexFactory;
+};
+
+
+class FOverlaySceneProxy : public FPrimitiveSceneProxy
+{
+public:
+
+	FOverlaySceneProxy( UOverlayComponent* Component )
+		: FPrimitiveSceneProxy( Component ),
+		  MaterialRelevance( Component->GetMaterialRelevance( GetScene().GetFeatureLevel() ) )
+	{
+		// @todo perf: Render state is marked dirty every time a new element is added, and a new scene proxy will be created.
+		// Perhaps we can instead create the render buffers as BUF_Dynamic and amend them directly on the render thread.
+		// However they will still require a reallocation if the number of elements grows, so this may not be a worthwhile optimization.
+
+		// Initialize lines.
+		// Lines are represented as two tris of zero thickness.
+		// The vertex normals are used to hold the normalized line direction + a sign determining the direction in which the
+		// material should thicken the polys.
+		if( Component->Lines.Num() > 0 )
+		{
+			MeshBatchDatas.Emplace();
+			FMeshBatchData& MeshBatchData = MeshBatchDatas.Last();
+
+			MeshBatchData.MaterialProxy = Component->GetMaterial( 0 )->GetRenderProxy( IsSelected() );
+			MeshBatchData.VertexFactory.Init( &MeshBatchData.VertexBuffer );
+			BeginInitResource( &MeshBatchData.VertexFactory );
+
+			int32 VertexBufferIndex = MeshBatchData.VertexBuffer.Vertices.AddUninitialized( Component->Lines.Num() * 4 );
+			int32 IndexBufferIndex = MeshBatchData.IndexBuffer.Indices.AddUninitialized( Component->Lines.Num() * 6 );
+
+			for( const FOverlayLine& OverlayLine : Component->Lines )
+			{
+				const FVector LineDirection = ( OverlayLine.End - OverlayLine.Start ).GetSafeNormal();
+				const FVector2D UV( OverlayLine.Thickness, 0.0f );
+
+				MeshBatchData.VertexBuffer.Vertices[ VertexBufferIndex + 0 ] = FOverlayVertex( OverlayLine.Start, UV, FVector::ZeroVector, -LineDirection, OverlayLine.Color );
+				MeshBatchData.VertexBuffer.Vertices[ VertexBufferIndex + 1 ] = FOverlayVertex( OverlayLine.End, UV, FVector::ZeroVector, -LineDirection, OverlayLine.Color );
+				MeshBatchData.VertexBuffer.Vertices[ VertexBufferIndex + 2 ] = FOverlayVertex( OverlayLine.End, UV, FVector::ZeroVector, LineDirection, OverlayLine.Color );
+				MeshBatchData.VertexBuffer.Vertices[ VertexBufferIndex + 3 ] = FOverlayVertex( OverlayLine.Start, UV, FVector::ZeroVector, LineDirection, OverlayLine.Color );
+
+				MeshBatchData.IndexBuffer.Indices[ IndexBufferIndex + 0 ] = VertexBufferIndex + 0;
+				MeshBatchData.IndexBuffer.Indices[ IndexBufferIndex + 1 ] = VertexBufferIndex + 1;
+				MeshBatchData.IndexBuffer.Indices[ IndexBufferIndex + 2 ] = VertexBufferIndex + 2;
+				MeshBatchData.IndexBuffer.Indices[ IndexBufferIndex + 3 ] = VertexBufferIndex + 2;
+				MeshBatchData.IndexBuffer.Indices[ IndexBufferIndex + 4 ] = VertexBufferIndex + 3;
+				MeshBatchData.IndexBuffer.Indices[ IndexBufferIndex + 5 ] = VertexBufferIndex + 0;
+
+				VertexBufferIndex += 4;
+				IndexBufferIndex += 6;
+			}
+
+			BeginInitResource( &MeshBatchData.VertexBuffer );
+			BeginInitResource( &MeshBatchData.IndexBuffer );
+		}
+
+		// Initialize points.
+		// Points are represented as two tris, all of whose vertices are coincident.
+		// The material then offsets them according to the signs of the vertex normals in a camera facing orientation.
+		// Size of the point is given by U0.
+		if( Component->Points.Num() > 0 )
+		{
+			MeshBatchDatas.Emplace();
+			FMeshBatchData& MeshBatchData = MeshBatchDatas.Last();
+
+			MeshBatchData.MaterialProxy = Component->GetMaterial( 1 )->GetRenderProxy( IsSelected() );
+			MeshBatchData.VertexFactory.Init( &MeshBatchData.VertexBuffer );
+			BeginInitResource( &MeshBatchData.VertexFactory );
+
+			int32 VertexBufferIndex = MeshBatchData.VertexBuffer.Vertices.AddUninitialized( Component->Points.Num() * 4 );
+			int32 IndexBufferIndex = MeshBatchData.IndexBuffer.Indices.AddUninitialized( Component->Points.Num() * 6 );
+
+			for( const FOverlayPoint& OverlayPoint : Component->Points )
+			{
+				const FVector2D UV( OverlayPoint.Size, 0.0f );
+
+				MeshBatchData.VertexBuffer.Vertices[ VertexBufferIndex + 0 ] = FOverlayVertex( OverlayPoint.Position, UV, FVector::ZeroVector, FVector( 1.0f, -1.0f, 0.0f ), OverlayPoint.Color );
+				MeshBatchData.VertexBuffer.Vertices[ VertexBufferIndex + 1 ] = FOverlayVertex( OverlayPoint.Position, UV, FVector::ZeroVector, FVector( 1.0f, 1.0f, 0.0f ), OverlayPoint.Color );
+				MeshBatchData.VertexBuffer.Vertices[ VertexBufferIndex + 2 ] = FOverlayVertex( OverlayPoint.Position, UV, FVector::ZeroVector, FVector( -1.0f, 1.0f, 0.0f ), OverlayPoint.Color );
+				MeshBatchData.VertexBuffer.Vertices[ VertexBufferIndex + 3 ] = FOverlayVertex( OverlayPoint.Position, UV, FVector::ZeroVector, FVector( -1.0f, -1.0f, 0.0f ), OverlayPoint.Color );
+
+				MeshBatchData.IndexBuffer.Indices[ IndexBufferIndex + 0 ] = VertexBufferIndex + 0;
+				MeshBatchData.IndexBuffer.Indices[ IndexBufferIndex + 1 ] = VertexBufferIndex + 1;
+				MeshBatchData.IndexBuffer.Indices[ IndexBufferIndex + 2 ] = VertexBufferIndex + 2;
+				MeshBatchData.IndexBuffer.Indices[ IndexBufferIndex + 3 ] = VertexBufferIndex + 2;
+				MeshBatchData.IndexBuffer.Indices[ IndexBufferIndex + 4 ] = VertexBufferIndex + 3;
+				MeshBatchData.IndexBuffer.Indices[ IndexBufferIndex + 5 ] = VertexBufferIndex + 0;
+
+				VertexBufferIndex += 4;
+				IndexBufferIndex += 6;
+			}
+
+			BeginInitResource( &MeshBatchData.VertexBuffer );
+			BeginInitResource( &MeshBatchData.IndexBuffer );
+		}
+
+		// Triangles
+		int32 MaterialIndex = 2;
+		for( const auto& MaterialTriangles : Component->TrianglesByMaterial )
+		{
+			MeshBatchDatas.Emplace();
+			FMeshBatchData& MeshBatchData = MeshBatchDatas.Last();
+
+			MeshBatchData.MaterialProxy = Component->GetMaterial( MaterialIndex )->GetRenderProxy( IsSelected() );
+			MeshBatchData.VertexFactory.Init( &MeshBatchData.VertexBuffer );
+			BeginInitResource( &MeshBatchData.VertexFactory );
+
+			int32 VertexBufferIndex = MeshBatchData.VertexBuffer.Vertices.AddUninitialized( MaterialTriangles.Num() * 3 );
+			int32 IndexBufferIndex = MeshBatchData.IndexBuffer.Indices.AddUninitialized( MaterialTriangles.Num() * 3 );
+
+			for( const FOverlayTriangle& Triangle : MaterialTriangles )
+			{
+				MeshBatchData.VertexBuffer.Vertices[ VertexBufferIndex + 0 ] = FOverlayVertex( Triangle.Vertex0.Position, Triangle.Vertex0.UV, FVector( 1, 0, 0 ), Triangle.Vertex0.Normal, Triangle.Vertex0.Color );
+				MeshBatchData.VertexBuffer.Vertices[ VertexBufferIndex + 1 ] = FOverlayVertex( Triangle.Vertex1.Position, Triangle.Vertex1.UV, FVector( 1, 0, 0 ), Triangle.Vertex1.Normal, Triangle.Vertex1.Color );
+				MeshBatchData.VertexBuffer.Vertices[ VertexBufferIndex + 2 ] = FOverlayVertex( Triangle.Vertex2.Position, Triangle.Vertex2.UV, FVector( 1, 0, 0 ), Triangle.Vertex2.Normal, Triangle.Vertex2.Color );
+
+				MeshBatchData.IndexBuffer.Indices[ IndexBufferIndex + 0 ] = VertexBufferIndex + 0;
+				MeshBatchData.IndexBuffer.Indices[ IndexBufferIndex + 1 ] = VertexBufferIndex + 1;
+				MeshBatchData.IndexBuffer.Indices[ IndexBufferIndex + 2 ] = VertexBufferIndex + 2;
+
+				VertexBufferIndex += 3;
+				IndexBufferIndex += 3;
+			}
+
+			BeginInitResource( &MeshBatchData.VertexBuffer );
+			BeginInitResource( &MeshBatchData.IndexBuffer );
+		}
+	}
+
+	virtual ~FOverlaySceneProxy()
+	{
+		for( FMeshBatchData& MeshBatchData : MeshBatchDatas )
+		{
+			MeshBatchData.VertexBuffer.ReleaseResource();
+			MeshBatchData.IndexBuffer.ReleaseResource();
+			MeshBatchData.VertexFactory.ReleaseResource();
+		}
+	}
+
+	virtual void GetDynamicMeshElements( const TArray<const FSceneView*>& Views, const FSceneViewFamily& ViewFamily, uint32 VisibilityMap, FMeshElementCollector& Collector ) const override
+	{
+		QUICK_SCOPE_CYCLE_COUNTER( STAT_OverlaySceneProxy_GetDynamicMeshElements );
+
+		for( int32 ViewIndex = 0; ViewIndex < Views.Num(); ViewIndex++ )
+		{
+			if( VisibilityMap & ( 1 << ViewIndex ) )
+			{
+				for( const FMeshBatchData& MeshBatchData : MeshBatchDatas )
+				{
+					const FSceneView* View = Views[ ViewIndex ];
+					FMeshBatch& Mesh = Collector.AllocateMesh();
+					FMeshBatchElement& BatchElement = Mesh.Elements[ 0 ];
+					BatchElement.IndexBuffer = &MeshBatchData.IndexBuffer;
+					Mesh.bWireframe = false;
+					Mesh.VertexFactory = &MeshBatchData.VertexFactory;
+					Mesh.MaterialRenderProxy = MeshBatchData.MaterialProxy;
+					BatchElement.PrimitiveUniformBuffer = CreatePrimitiveUniformBufferImmediate( GetLocalToWorld(), GetBounds(), GetLocalBounds(), true, UseEditorDepthTest() );
+					BatchElement.FirstIndex = 0;
+					BatchElement.NumPrimitives = MeshBatchData.IndexBuffer.Indices.Num() / 3;
+					BatchElement.MinVertexIndex = 0;
+					BatchElement.MaxVertexIndex = MeshBatchData.VertexBuffer.Vertices.Num() - 1;
+					Mesh.ReverseCulling = IsLocalToWorldDeterminantNegative();
+					Mesh.Type = PT_TriangleList;
+					Mesh.DepthPriorityGroup = SDPG_World;
+					Mesh.bCanApplyViewModeOverrides = false;
+					Collector.AddMesh( ViewIndex, Mesh );
+				}
+			}
+		}
+	}
+
+	virtual FPrimitiveViewRelevance GetViewRelevance( const FSceneView* View ) const override
+	{
+		FPrimitiveViewRelevance Result;
+		Result.bDrawRelevance = IsShown( View );
+		Result.bShadowRelevance = IsShadowCast( View );
+		Result.bDynamicRelevance = true;
+		Result.bRenderInMainPass = ShouldRenderInMainPass();
+		Result.bUsesLightingChannels = GetLightingChannelMask() != GetDefaultLightingChannelMask();
+		Result.bRenderCustomDepth = ShouldRenderCustomDepth();
+		MaterialRelevance.SetPrimitiveViewRelevance( Result );
+		return Result;
+	}
+
+	virtual bool CanBeOccluded() const override
+	{
+		return !MaterialRelevance.bDisableDepthTest;
+	}
+
+	virtual uint32 GetMemoryFootprint() const override { return sizeof( *this ) + GetAllocatedSize(); }
+
+	uint32 GetAllocatedSize() const { return FPrimitiveSceneProxy::GetAllocatedSize(); }
+
+private:
+
+	TArray<FMeshBatchData> MeshBatchDatas;
+	FMaterialRelevance MaterialRelevance;
+};
+
+
+
+UOverlayComponent::UOverlayComponent( const FObjectInitializer& ObjectInitializer )
+{
+	CastShadow = false;
+	bSelectable = false;
+	PrimaryComponentTick.bCanEverTick = false;
+	bBoundsDirty = true;
+
+	SetCollisionProfileName( UCollisionProfile::NoCollision_ProfileName );
+}
+
+
+void UOverlayComponent::SetLineMaterial( UMaterialInterface* InLineMaterial )
+{
+	LineMaterial = InLineMaterial;
+	SetMaterial( 0, InLineMaterial );
+}
+
+
+void UOverlayComponent::SetPointMaterial( UMaterialInterface* InPointMaterial )
+{
+	PointMaterial = InPointMaterial;
+	SetMaterial( 1, InPointMaterial );
+}
+
+
+void UOverlayComponent::Clear()
+{
+	Lines.Reset();
+	Points.Reset();
+	Triangles.Reset();
+	for( int32 Index = 0; Index < TrianglesByMaterial.Num(); Index++ )
+	{
+		SetMaterial( Index + 2, nullptr );
+	}
+	TrianglesByMaterial.Reset();
+	MaterialToIndex.Reset();
+	MarkRenderStateDirty();
+	bBoundsDirty = true;
+}
+
+
+FOverlayLineID UOverlayComponent::AddLine( const FOverlayLine& OverlayLine )
+{
+	const FOverlayLineID ID( Lines.Add( OverlayLine ) );
+	MarkRenderStateDirty();
+	bBoundsDirty = true;
+	return ID;
+}
+
+
+void UOverlayComponent::InsertLine( const FOverlayLineID ID, const FOverlayLine& OverlayLine )
+{
+	Lines.Insert( ID.GetValue(), OverlayLine ); 
+	MarkRenderStateDirty();
+	bBoundsDirty = true;
+}
+
+
+void UOverlayComponent::SetLineColor( const FOverlayLineID ID, const FColor& NewColor )
+{
+	FOverlayLine& OverlayLine = Lines[ ID.GetValue() ];
+	OverlayLine.Color = NewColor;
+	MarkRenderStateDirty();
+}
+
+
+void UOverlayComponent::SetLineThickness( const FOverlayLineID ID, const float NewThickness )
+{
+	FOverlayLine& OverlayLine = Lines[ ID.GetValue() ];
+	OverlayLine.Thickness = NewThickness;
+	MarkRenderStateDirty();
+}
+
+
+void UOverlayComponent::RemoveLine( const FOverlayLineID ID )
+{
+	Lines.RemoveAt( ID.GetValue() );
+	MarkRenderStateDirty();
+	bBoundsDirty = true;
+}
+
+
+bool UOverlayComponent::IsLineValid( const FOverlayLineID ID ) const
+{
+	return Lines.IsAllocated( ID.GetValue() );
+}
+
+
+FOverlayPointID UOverlayComponent::AddPoint( const FOverlayPoint& OverlayPoint )
+{
+	const FOverlayPointID ID( Points.Add( OverlayPoint ) );
+	MarkRenderStateDirty();
+	bBoundsDirty = true;
+	return ID;
+}
+
+
+void UOverlayComponent::InsertPoint( const FOverlayPointID ID, const FOverlayPoint& OverlayPoint )
+{
+	Points.Insert( ID.GetValue(), OverlayPoint );
+	MarkRenderStateDirty();
+	bBoundsDirty = true;
+}
+
+
+void UOverlayComponent::SetPointColor( const FOverlayPointID ID, const FColor& NewColor )
+{
+	FOverlayPoint& OverlayPoint = Points[ ID.GetValue() ];
+	OverlayPoint.Color = NewColor;
+	MarkRenderStateDirty();
+}
+
+
+void UOverlayComponent::SetPointSize( const FOverlayPointID ID, const float NewSize )
+{
+	FOverlayPoint& OverlayPoint = Points[ ID.GetValue() ];
+	OverlayPoint.Size = NewSize;
+	MarkRenderStateDirty();
+}
+
+
+void UOverlayComponent::RemovePoint( const FOverlayPointID ID )
+{
+	Points.RemoveAt( ID.GetValue() );
+	MarkRenderStateDirty();
+	bBoundsDirty = true;
+}
+
+
+bool UOverlayComponent::IsPointValid( const FOverlayPointID ID ) const
+{
+	return Points.IsAllocated( ID.GetValue() );
+}
+
+
+int32 UOverlayComponent::FindOrAddMaterialIndex( UMaterialInterface* Material )
+{
+	const int32* MaterialIndexPtr = MaterialToIndex.Find( Material );
+	if( MaterialIndexPtr == nullptr )
+	{
+		const int32 MaterialIndex = TrianglesByMaterial.Add( TSparseArray<FOverlayTriangle>() );
+		MaterialToIndex.Add( Material, MaterialIndex );
+		SetMaterial( MaterialIndex + 2, Material );
+		return MaterialIndex;
+	}
+
+	return *MaterialIndexPtr;
+}
+
+
+FOverlayTriangleID UOverlayComponent::AddTriangle( const FOverlayTriangle& OverlayTriangle )
+{
+	const int32 MaterialIndex = FindOrAddMaterialIndex( OverlayTriangle.Material );
+	const int32 IndexByMaterial = TrianglesByMaterial[ MaterialIndex ].Add( OverlayTriangle );
+	const FOverlayTriangleID ID( Triangles.Add( MakeTuple( MaterialIndex, IndexByMaterial ) ) );
+	MarkRenderStateDirty();
+	bBoundsDirty = true;
+	return ID;
+}
+
+
+void UOverlayComponent::InsertTriangle( const FOverlayTriangleID ID, const FOverlayTriangle& OverlayTriangle )
+{
+	const int32 MaterialIndex = FindOrAddMaterialIndex( OverlayTriangle.Material );
+	const int32 IndexByMaterial = TrianglesByMaterial[ MaterialIndex ].Add( OverlayTriangle );
+	Triangles.Insert( ID.GetValue(), MakeTuple( MaterialIndex, IndexByMaterial ) );
+	MarkRenderStateDirty();
+	bBoundsDirty = true;
+}
+
+
+void UOverlayComponent::RemoveTriangle( const FOverlayTriangleID ID )
+{
+	const TTuple<int32, int32> MaterialAndTriangleIndex = Triangles[ ID.GetValue() ];
+	const int32 MaterialIndex = MaterialAndTriangleIndex.Get<0>();
+	const int32 IndexByMaterial = MaterialAndTriangleIndex.Get<1>();
+	TSparseArray<FOverlayTriangle>& Container = TrianglesByMaterial[ MaterialIndex ];
+	Container.RemoveAt( IndexByMaterial );
+	if( Container.Num() == 0 )
+	{
+		TrianglesByMaterial.RemoveAt( MaterialIndex );
+		SetMaterial( MaterialIndex + 2, nullptr );
+	}
+	Triangles.RemoveAt( ID.GetValue() );
+	MarkRenderStateDirty();
+	bBoundsDirty = true;
+}
+
+
+bool UOverlayComponent::IsTriangleValid( const FOverlayTriangleID ID ) const
+{
+	return Triangles.IsAllocated( ID.GetValue() );
+}
+
+
+FPrimitiveSceneProxy* UOverlayComponent::CreateSceneProxy()
+{
+	if( Lines.Num() > 0 || Points.Num() > 0 || Triangles.Num() > 0 )
+	{
+		return new FOverlaySceneProxy( this );
+	}
+
+	return nullptr;
+}
+
+
+int32 UOverlayComponent::GetNumMaterials() const
+{
+	return TrianglesByMaterial.GetMaxIndex() + 2;
+}
+
+
+FBoxSphereBounds UOverlayComponent::CalcBounds( const FTransform& LocalToWorld ) const
+{
+	if( bBoundsDirty )
+	{
+		FBox Box( ForceInit );
+
+		for( const FOverlayLine& Line : Lines )
+		{
+			Box += Line.Start;
+			Box += Line.End;
+		}
+
+		for( const FOverlayPoint& Point : Points )
+		{
+			Box += Point.Position;
+		}
+
+		for( const TSparseArray<FOverlayTriangle>& TriangleArray : TrianglesByMaterial )
+		{
+			for( const FOverlayTriangle& Triangle : TriangleArray )
+			{
+				Box += Triangle.Vertex0.Position;
+				Box += Triangle.Vertex1.Position;
+				Box += Triangle.Vertex2.Position;
+			}
+		}
+
+		Bounds = FBoxSphereBounds( Box );
+		bBoundsDirty = false;
+	}
+
+	return Bounds.TransformBy( LocalToWorld );
+}
+
diff --git a/Engine/Plugins/Editor/MeshEditor/Source/MeshEditor/OverlayComponent.h b/Engine/Plugins/Editor/MeshEditor/Source/MeshEditor/OverlayComponent.h
new file mode 100644
--- /dev/null
+++ b/Engine/Plugins/Editor/MeshEditor/Source/MeshEditor/OverlayComponent.h
@@ -0,0 +1,275 @@
+// Copyright 1998-2017 Epic Games, Inc. All Rights Reserved.
+
+#pragma once
+
+#include "CoreMinimal.h"
+#include "UObject/ObjectMacros.h"
+#include "Components/MeshComponent.h"
+#include "Containers/SparseArray.h"
+#include "RenderResource.h"
+#include "PackedNormal.h"
+#include "OverlayComponent.generated.h"
+
+class FPrimitiveSceneProxy;
+
+
+USTRUCT()
+struct FOverlayLine
+{
+	GENERATED_BODY()
+
+	FOverlayLine() {}
+
+	FOverlayLine( const FVector& InStart, const FVector& InEnd, const FColor& InColor, const float InThickness )
+		: Start( InStart ),
+		  End( InEnd ),
+		  Color( InColor ),
+		  Thickness( InThickness )
+	{
+	}
+
+	UPROPERTY()
+	FVector Start;
+
+	UPROPERTY()
+	FVector End;
+
+	UPROPERTY()
+	FColor Color;
+
+	UPROPERTY()
+	float Thickness;
+};
+
+
+USTRUCT()
+struct FOverlayPoint
+{
+	GENERATED_BODY()
+
+	FOverlayPoint() {}
+
+	FOverlayPoint( const FVector& InPosition, const FColor& InColor, const float InSize )
+		: Position( InPosition ),
+		  Color( InColor ),
+		  Size( InSize )
+	{}
+
+	UPROPERTY()
+	FVector Position;
+
+	UPROPERTY()
+	FColor Color;
+
+	UPROPERTY()
+	float Size;
+};
+
+
+USTRUCT()
+struct FOverlayTriangleVertex
+{
+	GENERATED_BODY()
+
+	FOverlayTriangleVertex() {}
+
+	FOverlayTriangleVertex( const FVector& InPosition, const FVector2D& InUV, const FVector& InNormal, const FColor& InColor )
+		: Position( InPosition ),
+		  UV( InUV ),
+		  Normal( InNormal ),
+		  Color( InColor )
+	{
+	}
+
+	UPROPERTY()
+	FVector Position;
+
+	UPROPERTY()
+	FVector2D UV;
+
+	UPROPERTY()
+	FVector Normal;
+
+	UPROPERTY()
+	FColor Color;
+};
+
+
+USTRUCT()
+struct FOverlayTriangle
+{
+	GENERATED_BODY()
+
+	FOverlayTriangle() {}
+
+	FOverlayTriangle( UMaterialInterface* InMaterial, const FOverlayTriangleVertex& InVertex0, const FOverlayTriangleVertex& InVertex1, const FOverlayTriangleVertex& InVertex2 )
+		: Material( InMaterial ),
+		  Vertex0( InVertex0 ),
+		  Vertex1( InVertex1 ),
+		  Vertex2( InVertex2 )
+	{
+	}
+
+	UPROPERTY()
+	UMaterialInterface* Material;
+
+	UPROPERTY()
+	FOverlayTriangleVertex Vertex0;
+
+	UPROPERTY()
+	FOverlayTriangleVertex Vertex1;
+
+	UPROPERTY()
+	FOverlayTriangleVertex Vertex2;
+};
+
+
+USTRUCT()
+struct FOverlayLineID
+{
+	GENERATED_BODY()
+
+	FOverlayLineID() : ID( 0 ) {}
+	explicit FOverlayLineID( const int32 InID ) : ID( InID ) {}
+
+	int32 GetValue() const { return ID; }
+
+	UPROPERTY()
+	int32 ID;
+};
+
+
+USTRUCT()
+struct FOverlayPointID
+{
+	GENERATED_BODY()
+
+	FOverlayPointID() : ID( 0 ) {}
+	explicit FOverlayPointID( const int32 InID ) : ID( InID ) {}
+
+	int32 GetValue() const { return ID; }
+
+	UPROPERTY()
+	int32 ID;
+};
+
+
+USTRUCT()
+struct FOverlayTriangleID
+{
+	GENERATED_BODY()
+
+	FOverlayTriangleID() : ID( 0 ) {}
+	explicit FOverlayTriangleID( const int32 InID ) : ID( InID ) {}
+
+	int32 GetValue() const { return ID; }
+
+	UPROPERTY()
+	int32 ID;
+};
+
+
+UCLASS()
+class UOverlayComponent : public UMeshComponent
+{
+	GENERATED_BODY()
+
+public:
+	/**
+	 * Default UObject constructor.
+	 */
+	UOverlayComponent( const FObjectInitializer& ObjectInitializer = FObjectInitializer::Get() );
+
+	/** Specify material which handles lines */
+	void SetLineMaterial( UMaterialInterface* InLineMaterial );
+
+	/** Specify material which handles points */
+	void SetPointMaterial( UMaterialInterface* InPointMaterial );
+
+	/** Clear all primitives */
+	void Clear();
+
+	/** Add a line to the overlay */
+	FOverlayLineID AddLine( const FOverlayLine& OverlayLine );
+
+	/** Insert a line with the given ID to the overlay */
+	void InsertLine( const FOverlayLineID ID, const FOverlayLine& OverlayLine );
+
+	/** Sets the color of a line */
+	void SetLineColor( const FOverlayLineID ID, const FColor& NewColor );
+
+	/** Sets the thickness of a line */
+	void SetLineThickness( const FOverlayLineID ID, const float NewThickness );
+
+	/** Remove a line from the overlay */
+	void RemoveLine( const FOverlayLineID ID );
+
+	/** Queries whether a line with the given ID exists */
+	bool IsLineValid( const FOverlayLineID ID ) const;
+
+	/** Add a point to the overlay */
+	FOverlayPointID AddPoint( const FOverlayPoint& OverlayPoint );
+
+	/** Insert a point with the given ID to the overlay */
+	void InsertPoint( const FOverlayPointID ID, const FOverlayPoint& OverlayPoint );
+
+	/** Sets the color of a point */
+	void SetPointColor( const FOverlayPointID ID, const FColor& NewColor );
+
+	/** Sets the size of a point */
+	void SetPointSize( const FOverlayPointID ID, const float NewSize );
+
+	/** Remove a point from the overlay */
+	void RemovePoint( const FOverlayPointID ID );
+
+	/** Queries whether a point with the given ID exists */
+	bool IsPointValid( const FOverlayPointID ID ) const;
+
+	/** Add a triangle to the overlay */
+	FOverlayTriangleID AddTriangle( const FOverlayTriangle& OverlayTriangle );
+
+	/** Insert a triangle with the given ID to the overlay */
+	void InsertTriangle( const FOverlayTriangleID ID, const FOverlayTriangle& OverlayTriangle );
+
+	/** Remove a triangle from the overlay */
+	void RemoveTriangle( const FOverlayTriangleID ID );
+
+	/** Queries whether a triangle with the given ID exists */
+	bool IsTriangleValid( const FOverlayTriangleID ID ) const;
+
+private:
+
+	//~ Begin UPrimitiveComponent Interface.
+	virtual FPrimitiveSceneProxy* CreateSceneProxy() override;
+	//~ End UPrimitiveComponent Interface.
+
+	//~ Begin UMeshComponent Interface.
+	virtual int32 GetNumMaterials() const override;
+	//~ End UMeshComponent Interface.
+
+	//~ Begin USceneComponent Interface.
+	virtual FBoxSphereBounds CalcBounds( const FTransform& LocalToWorld ) const override;
+	//~ Begin USceneComponent Interface.
+
+	int32 FindOrAddMaterialIndex( UMaterialInterface* Material );
+
+	UPROPERTY()
+	const UMaterialInterface* LineMaterial;
+
+	UPROPERTY()
+	const UMaterialInterface* PointMaterial;
+
+	UPROPERTY()
+	mutable FBoxSphereBounds Bounds;
+
+	UPROPERTY()
+	mutable bool bBoundsDirty;
+
+	TSparseArray<FOverlayLine> Lines;
+	TSparseArray<FOverlayPoint> Points;
+	TSparseArray<TTuple<int32, int32>> Triangles;
+	TSparseArray<TSparseArray<FOverlayTriangle>> TrianglesByMaterial;
+	TMap<UMaterialInterface*, int32> MaterialToIndex;
+
+	friend class FOverlaySceneProxy;
+};
diff --git a/Engine/Plugins/Editor/MeshEditor/Source/MeshEditor/WireframeMeshComponent.cpp b/Engine/Plugins/Editor/MeshEditor/Source/MeshEditor/WireframeMeshComponent.cpp
new file mode 100644
--- /dev/null
+++ b/Engine/Plugins/Editor/MeshEditor/Source/MeshEditor/WireframeMeshComponent.cpp
@@ -0,0 +1,483 @@
+// Copyright 1998-2017 Epic Games, Inc. All Rights Reserved. 
+
+#include "WireframeMeshComponent.h"
+#include "RenderingThread.h"
+#include "PrimitiveViewRelevance.h"
+#include "PrimitiveSceneProxy.h"
+#include "VertexFactory.h"
+#include "MaterialShared.h"
+#include "Engine/CollisionProfile.h"
+#include "Materials/Material.h"
+#include "LocalVertexFactory.h"
+#include "SceneManagement.h"
+#include "DynamicMeshBuilder.h"
+#include "EngineGlobals.h"
+#include "Engine/Engine.h"
+
+
+void UWireframeMesh::Reset()
+{
+	Vertices.Reset();
+	Polygons.Reset();
+	Edges.Reset();
+	EdgeInstances.Reset();
+	bBoundsDirty = true;
+}
+
+
+void UWireframeMesh::AddVertex( const FVertexID VertexID )
+{
+	Vertices.Insert( VertexID.GetValue(), FWireframeVertex() );
+}
+
+
+void UWireframeMesh::SetVertexPosition( const FVertexID VertexID, const FVector& Position )
+{
+	Vertices[ VertexID.GetValue() ].Position = Position;
+}
+
+
+void UWireframeMesh::RemoveVertex( const FVertexID VertexID )
+{
+	Vertices.RemoveAt( VertexID.GetValue() );
+	bBoundsDirty = true;
+}
+
+
+void UWireframeMesh::AddPolygon( const FPolygonID PolygonID )
+{
+	Polygons.Insert( PolygonID.GetValue(), FWireframePolygon() );
+}
+
+
+void UWireframeMesh::SetPolygonNormal( const FPolygonID PolygonID, const FVector& Normal )
+{
+	Polygons[ PolygonID.GetValue() ].PolygonNormal = Normal;
+}
+
+
+void UWireframeMesh::RemovePolygon( const FPolygonID PolygonID )
+{
+	Polygons.RemoveAt( PolygonID.GetValue() );
+}
+
+
+void UWireframeMesh::AddEdge( const FEdgeID EdgeID )
+{
+	Edges.Insert( EdgeID.GetValue(), FWireframeEdge() );
+}
+
+
+void UWireframeMesh::SetEdgeVertices( const FEdgeID EdgeID, const FVertexID Vertex0, const FVertexID Vertex1 )
+{
+	Edges[ EdgeID.GetValue() ].StartVertex = Vertex0;
+	Edges[ EdgeID.GetValue() ].EndVertex = Vertex1;
+}
+
+
+void UWireframeMesh::SetEdgeColor( const FEdgeID EdgeID, const FColor Color )
+{
+	Edges[ EdgeID.GetValue() ].Color = Color;
+}
+
+
+void UWireframeMesh::RemoveEdge( const FEdgeID EdgeID )
+{
+	check( Edges[ EdgeID.GetValue() ].EdgeInstances.Num() == 0 );
+	Edges.RemoveAt( EdgeID.GetValue() );
+}
+
+
+void UWireframeMesh::AddEdgeInstance( const FEdgeID EdgeID, const FPolygonID PolygonID )
+{
+	const int32 Index = EdgeInstances.Add( FWireframeEdgeInstance() );
+	EdgeInstances[ Index ].PolygonID = PolygonID;
+	EdgeInstances[ Index ].EdgeID = EdgeID;
+	FWireframeEdge& Edge = Edges[ EdgeID.GetValue() ];
+	check( !Edge.EdgeInstances.Contains( Index ) );
+	Edge.EdgeInstances.Add( Index );
+}
+
+
+void UWireframeMesh::RemoveEdgeInstance( const FEdgeID EdgeID, const FPolygonID PolygonID )
+{
+	FWireframeEdge& Edge = Edges[ EdgeID.GetValue() ];
+
+	// Find edge instance index corresponding to an instance of this edge with the given polygon ID
+	const int32* EdgeInstanceIndexPtr = Edge.EdgeInstances.FindByPredicate( 
+		[ this, PolygonID ]( const int32 Index ) { return this->EdgeInstances[ Index ].PolygonID == PolygonID; }
+	);
+	check( EdgeInstanceIndexPtr != nullptr );
+	const int32 EdgeInstanceIndex = *EdgeInstanceIndexPtr;
+
+	// Remove the instance from the edge array of instances
+	const int32 NumberOfItemsRemoved = Edge.EdgeInstances.RemoveSwap( EdgeInstanceIndex );
+	verify( NumberOfItemsRemoved == 1 );
+
+	// If we are not removing the last element, we need to move the last element into the space vacated by the removed element
+	// and patch up the reference to it in the Edge.
+	const int32 IndexOfLastInstance = EdgeInstances.Num() - 1;
+	if( EdgeInstanceIndex != IndexOfLastInstance )
+	{
+		// First, get the last instance in the array, and the edge it is instancing
+		const FWireframeEdgeInstance& LastInstance = EdgeInstances[ IndexOfLastInstance ];
+		FWireframeEdge& EdgeForLastInstance = Edges[ LastInstance.EdgeID.GetValue() ];
+
+		// Now find the index of the last edge instance inside the array of instances for that edge.
+		// Edge instances are owned by exactly one edge.
+		const int32 EdgeLastInstanceIndex = EdgeForLastInstance.EdgeInstances.Find( IndexOfLastInstance );
+		check( EdgeLastInstanceIndex != INDEX_NONE );
+
+		// Fix up the index of the edge instance which has now changed
+		check( !EdgeForLastInstance.EdgeInstances.Contains( EdgeInstanceIndex ) );
+		EdgeForLastInstance.EdgeInstances[ EdgeLastInstanceIndex ] = EdgeInstanceIndex;
+
+		// Copy the last edge instance into the element vacated by the removed edge instance
+		EdgeInstances[ EdgeInstanceIndex ] = LastInstance;
+	}
+
+	// Reduce the array size by one
+	EdgeInstances.SetNum( IndexOfLastInstance );
+}
+
+
+const TArray<int32>& UWireframeMesh::GetEdgeInstanceIDs( const FEdgeID EdgeID ) const
+{
+	return Edges[ EdgeID.GetValue() ].EdgeInstances;
+}
+
+
+FBoxSphereBounds UWireframeMesh::GetBounds() const
+{
+	if( bBoundsDirty )
+	{
+		FBox Box( ForceInit );
+		for( const FWireframeVertex& Vertex : Vertices )
+		{
+			Box += Vertex.Position;
+		}
+
+		Bounds = FBoxSphereBounds( Box );
+		bBoundsDirty = false;
+	}
+
+	return Bounds;
+}
+
+
+void UWireframeMesh::InitResources()
+{
+	const int32 EdgeInstanceCount = EdgeInstances.Num();
+
+	VertexBuffer.Vertices.Reset();
+	IndexBuffer.Indices.Reset();
+
+	if( EdgeInstanceCount > 0 )
+	{
+		// An edge instance is a concrete instance of an edge on a particular polygon. Each edge instance is expanded into two triangles which form a camera facing quad.
+		// The quad has a thickness of zero: the material is responsible for giving it finite thickness, according to the edge and the camera direction.
+
+		// The polygons for the wireframe mesh uses a specific vertex format:
+		// The vertex normal is used to represent the edge direction (we use the normal rather than the tangent as it must not be modified by orthonormalization, when building the tangent basis).
+		// The vertex tangent is used to represent the edge normal, which is used for optional backface culling.
+		// The UV0 channel is set on a per-component basis, and contains various overrides per edge, to control the opacity and highlighting.
+
+		int32 VertexBufferIndex = VertexBuffer.Vertices.AddUninitialized( EdgeInstanceCount * 4 );
+		int32 IndexBufferIndex = IndexBuffer.Indices.AddUninitialized( EdgeInstanceCount * 6 );
+
+		for( const FWireframeEdgeInstance& EdgeInstance : EdgeInstances )
+		{
+			const FWireframeEdge& Edge = Edges[ EdgeInstance.EdgeID.GetValue() ];
+
+			const FVector& PolygonNormal = Polygons[ EdgeInstance.PolygonID.GetValue() ].PolygonNormal;
+			const FVector& StartVertex = Vertices[ Edge.StartVertex.GetValue() ].Position;
+			const FVector& EndVertex = Vertices[ Edge.EndVertex.GetValue() ].Position;
+			const FVector EdgeDirection = ( EndVertex - StartVertex ).GetSafeNormal();
+
+			VertexBuffer.Vertices[ VertexBufferIndex + 0 ] = FWireframeMeshVertex( StartVertex, PolygonNormal, -EdgeDirection, Edge.Color );
+			VertexBuffer.Vertices[ VertexBufferIndex + 1 ] = FWireframeMeshVertex( EndVertex, PolygonNormal, -EdgeDirection, Edge.Color );
+			VertexBuffer.Vertices[ VertexBufferIndex + 2 ] = FWireframeMeshVertex( EndVertex, PolygonNormal, EdgeDirection, Edge.Color );
+			VertexBuffer.Vertices[ VertexBufferIndex + 3 ] = FWireframeMeshVertex( StartVertex, PolygonNormal, EdgeDirection, Edge.Color );
+
+			IndexBuffer.Indices[ IndexBufferIndex + 0 ] = VertexBufferIndex + 0;
+			IndexBuffer.Indices[ IndexBufferIndex + 1 ] = VertexBufferIndex + 1;
+			IndexBuffer.Indices[ IndexBufferIndex + 2 ] = VertexBufferIndex + 2;
+			IndexBuffer.Indices[ IndexBufferIndex + 3 ] = VertexBufferIndex + 2;
+			IndexBuffer.Indices[ IndexBufferIndex + 4 ] = VertexBufferIndex + 3;
+			IndexBuffer.Indices[ IndexBufferIndex + 5 ] = VertexBufferIndex + 0;
+
+			VertexBufferIndex += 4;
+			IndexBufferIndex += 6;
+		}
+
+		// Enqueue initialization of render resource
+		BeginInitResource( &VertexBuffer );
+		BeginInitResource( &IndexBuffer );
+	}
+}
+
+
+void UWireframeMesh::ReleaseResources()
+{
+	if( IndexBuffer.Indices.Num() > 0 )
+	{
+		BeginReleaseResource( &VertexBuffer );
+		BeginReleaseResource( &IndexBuffer );
+	}
+}
+
+
+void FWireframeMeshVertexBuffer::InitRHI()
+{
+	FRHIResourceCreateInfo CreateInfo;
+	void* VertexBufferData = nullptr;
+	VertexBufferRHI = RHICreateAndLockVertexBuffer( Vertices.Num() * sizeof( FWireframeMeshVertex ), BUF_Static, CreateInfo, VertexBufferData );
+
+	// Copy the vertex data into the vertex buffer.		
+	FMemory::Memcpy( VertexBufferData, Vertices.GetData(), Vertices.Num() * sizeof( FWireframeMeshVertex ) );
+	RHIUnlockVertexBuffer( VertexBufferRHI );
+}
+
+
+void FWireframeMeshIndexBuffer::InitRHI()
+{
+	FRHIResourceCreateInfo CreateInfo;
+	void* Buffer = nullptr;
+	IndexBufferRHI = RHICreateAndLockIndexBuffer( sizeof( int32 ), Indices.Num() * sizeof( int32 ), BUF_Static, CreateInfo, Buffer );
+
+	// Write the indices to the index buffer.		
+	FMemory::Memcpy( Buffer, Indices.GetData(), Indices.Num() * sizeof( int32 ) );
+	RHIUnlockIndexBuffer( IndexBufferRHI );
+}
+
+
+void FWireframeMeshInstanceVertexBuffer::InitRHI()
+{
+	FRHIResourceCreateInfo CreateInfo;
+	void* VertexBufferData = nullptr;
+	VertexBufferRHI = RHICreateAndLockVertexBuffer( Vertices.Num() * sizeof( FWireframeMeshInstanceVertex ), BUF_Static, CreateInfo, VertexBufferData );
+
+	// Copy the vertex data into the vertex buffer.		
+	FMemory::Memcpy( VertexBufferData, Vertices.GetData(), Vertices.Num() * sizeof( FWireframeMeshInstanceVertex ) );
+	RHIUnlockVertexBuffer( VertexBufferRHI );
+}
+
+
+/** Vertex Factory */
+class FWireframeMeshVertexFactory : public FLocalVertexFactory
+{
+public:
+
+	FWireframeMeshVertexFactory()
+	{}
+
+	/** Init function that should only be called on render thread. */
+	void Init_RenderThread( const FWireframeMeshVertexBuffer* VertexBuffer, const FWireframeMeshInstanceVertexBuffer* InstanceVertexBuffer )
+	{
+		check( IsInRenderingThread() );
+
+		FDataType NewData;
+		NewData.PositionComponent = STRUCTMEMBER_VERTEXSTREAMCOMPONENT( VertexBuffer, FWireframeMeshVertex, Position, VET_Float3 );
+		NewData.TextureCoordinates.Add( FVertexStreamComponent( InstanceVertexBuffer, STRUCT_OFFSET( FWireframeMeshInstanceVertex, UV ), sizeof( FWireframeMeshInstanceVertex ), VET_Float2 ) );
+		NewData.TangentBasisComponents[ 0 ] = STRUCTMEMBER_VERTEXSTREAMCOMPONENT( VertexBuffer, FWireframeMeshVertex, TangentX, VET_PackedNormal );
+		NewData.TangentBasisComponents[ 1 ] = STRUCTMEMBER_VERTEXSTREAMCOMPONENT( VertexBuffer, FWireframeMeshVertex, TangentZ, VET_PackedNormal );
+		NewData.ColorComponent = STRUCTMEMBER_VERTEXSTREAMCOMPONENT( VertexBuffer, FWireframeMeshVertex, Color, VET_Color );
+
+		SetData( NewData );
+	}
+
+	/** Initialization */
+	void Init( const FWireframeMeshVertexBuffer* VertexBuffer, const FWireframeMeshInstanceVertexBuffer* InstanceVertexBuffer )
+	{
+		if( IsInRenderingThread() )
+		{
+			Init_RenderThread( VertexBuffer, InstanceVertexBuffer );
+		}
+		else
+		{
+			ENQUEUE_UNIQUE_RENDER_COMMAND_THREEPARAMETER(
+				InitWireframeMeshVertexFactory,
+				FWireframeMeshVertexFactory*, VertexFactory, this,
+				const FWireframeMeshVertexBuffer*, VertexBuffer, VertexBuffer,
+				const FWireframeMeshInstanceVertexBuffer*, InstanceVertexBuffer, InstanceVertexBuffer,
+				{
+					VertexFactory->Init_RenderThread( VertexBuffer, InstanceVertexBuffer );
+				});
+		}	
+	}
+};
+
+
+/** Scene proxy */
+class FWireframeMeshSceneProxy : public FPrimitiveSceneProxy
+{
+public:
+
+	FWireframeMeshSceneProxy( UWireframeMeshComponent* Component )
+		: FPrimitiveSceneProxy( Component ),
+		  MaterialRelevance( Component->GetMaterialRelevance( GetScene().GetFeatureLevel() ) )
+	{
+		check( Component->GetWireframeMesh() );
+
+		VertexBuffer = &Component->GetWireframeMesh()->VertexBuffer;
+		IndexBuffer = &Component->GetWireframeMesh()->IndexBuffer;
+
+		// Init vertex factory
+		VertexFactory.Init( &Component->GetWireframeMesh()->VertexBuffer, &InstanceVertexBuffer );
+		BeginInitResource( &VertexFactory );
+
+		// Init instance data
+		InstanceVertexBuffer.Vertices.AddZeroed( Component->GetWireframeMesh()->GetNumEdgeInstances() * 4 );
+
+		for( const FEdgeID HiddenEdgeID : Component->HiddenEdgeIDs )
+		{
+			const TArray<int32>& EdgeInstanceIndices = Component->GetWireframeMesh()->GetEdgeInstanceIDs( HiddenEdgeID );
+			for( const int32 EdgeInstanceIndex : EdgeInstanceIndices )
+			{
+				// Hide an edge by setting U of all of its instance vertices to 1.0
+				InstanceVertexBuffer.Vertices[ EdgeInstanceIndex * 4 + 0 ].UV = FVector2D( 1.0f, 0.0f );
+				InstanceVertexBuffer.Vertices[ EdgeInstanceIndex * 4 + 1 ].UV = FVector2D( 1.0f, 0.0f );
+				InstanceVertexBuffer.Vertices[ EdgeInstanceIndex * 4 + 2 ].UV = FVector2D( 1.0f, 0.0f );
+				InstanceVertexBuffer.Vertices[ EdgeInstanceIndex * 4 + 3 ].UV = FVector2D( 1.0f, 0.0f );
+			}
+		}
+
+		BeginInitResource( &InstanceVertexBuffer );
+
+		// Grab material
+		Material = Component->GetMaterial( 0 );
+		if( Material == nullptr )
+		{
+			Material = UMaterial::GetDefaultMaterial( MD_Surface );
+		}
+	}
+
+	virtual ~FWireframeMeshSceneProxy()
+	{
+		InstanceVertexBuffer.ReleaseResource();
+		VertexFactory.ReleaseResource();
+	}
+
+	virtual void GetDynamicMeshElements( const TArray<const FSceneView*>& Views, const FSceneViewFamily& ViewFamily, uint32 VisibilityMap, FMeshElementCollector& Collector ) const override
+	{
+		QUICK_SCOPE_CYCLE_COUNTER( STAT_WireframeMeshSceneProxy_GetDynamicMeshElements );
+
+		FMaterialRenderProxy* MaterialProxy = Material->GetRenderProxy( IsSelected() );
+
+		for( int32 ViewIndex = 0; ViewIndex < Views.Num(); ViewIndex++ )
+		{
+			if( VisibilityMap & ( 1 << ViewIndex ) )
+			{
+				const FSceneView* View = Views[ ViewIndex ];
+				FMeshBatch& Mesh = Collector.AllocateMesh();
+				FMeshBatchElement& BatchElement = Mesh.Elements[ 0 ];
+				BatchElement.IndexBuffer = IndexBuffer;
+				Mesh.bWireframe = false;
+				Mesh.VertexFactory = &VertexFactory;
+				Mesh.MaterialRenderProxy = MaterialProxy;
+				BatchElement.PrimitiveUniformBuffer = CreatePrimitiveUniformBufferImmediate( GetLocalToWorld(), GetBounds(), GetLocalBounds(), true, UseEditorDepthTest() );
+				BatchElement.FirstIndex = 0;
+				BatchElement.NumPrimitives = IndexBuffer->Indices.Num() / 3;
+				BatchElement.MinVertexIndex = 0;
+				BatchElement.MaxVertexIndex = VertexBuffer->Vertices.Num() - 1;
+				Mesh.ReverseCulling = IsLocalToWorldDeterminantNegative();
+				Mesh.Type = PT_TriangleList;
+				Mesh.DepthPriorityGroup = SDPG_World;
+				Mesh.bCanApplyViewModeOverrides = false;
+				Collector.AddMesh( ViewIndex, Mesh );
+			}
+		}
+	}
+
+	virtual FPrimitiveViewRelevance GetViewRelevance( const FSceneView* View ) const override
+	{
+		FPrimitiveViewRelevance Result;
+		Result.bDrawRelevance = IsShown( View );
+		Result.bShadowRelevance = IsShadowCast( View );
+		Result.bDynamicRelevance = true;
+		Result.bRenderInMainPass = ShouldRenderInMainPass();
+		Result.bUsesLightingChannels = GetLightingChannelMask() != GetDefaultLightingChannelMask();
+		Result.bRenderCustomDepth = ShouldRenderCustomDepth();
+		MaterialRelevance.SetPrimitiveViewRelevance( Result );
+		return Result;
+	}
+
+	virtual bool CanBeOccluded() const override
+	{
+		return !MaterialRelevance.bDisableDepthTest;
+	}
+
+	virtual uint32 GetMemoryFootprint() const override { return sizeof( *this ) + GetAllocatedSize(); }
+
+	uint32 GetAllocatedSize() const { return FPrimitiveSceneProxy::GetAllocatedSize(); }
+
+private:
+
+	UMaterialInterface* Material;
+	const FWireframeMeshVertexBuffer* VertexBuffer;
+	const FWireframeMeshIndexBuffer* IndexBuffer;
+	FWireframeMeshInstanceVertexBuffer InstanceVertexBuffer;
+	FWireframeMeshVertexFactory VertexFactory;
+
+	FMaterialRelevance MaterialRelevance;
+};
+
+
+//////////////////////////////////////////////////////////////////////////
+
+UWireframeMeshComponent::UWireframeMeshComponent( const FObjectInitializer& ObjectInitializer )
+{
+	CastShadow = false;
+	bSelectable = false;
+	PrimaryComponentTick.bCanEverTick = false;
+	WireframeMesh = nullptr;
+
+	SetCollisionProfileName( UCollisionProfile::NoCollision_ProfileName );
+}
+
+
+void UWireframeMeshComponent::ShowAllEdges()
+{
+	HiddenEdgeIDs.Reset();
+}
+
+
+void UWireframeMeshComponent::SetEdgeVisibility( const FEdgeID EdgeID, const bool bEdgeVisible )
+{
+	if( bEdgeVisible )
+	{
+		HiddenEdgeIDs.Remove( EdgeID );
+	}
+	else
+	{
+		HiddenEdgeIDs.Add( EdgeID );
+	}
+}
+
+
+FPrimitiveSceneProxy* UWireframeMeshComponent::CreateSceneProxy()
+{
+	if( WireframeMesh && WireframeMesh->EdgeInstances.Num() > 0 )
+	{
+		return new FWireframeMeshSceneProxy( this );
+	}
+
+	return nullptr;
+}
+
+
+int32 UWireframeMeshComponent::GetNumMaterials() const
+{
+	return 1;
+}
+
+
+FBoxSphereBounds UWireframeMeshComponent::CalcBounds( const FTransform& LocalToWorld ) const
+{
+	if( WireframeMesh )
+	{
+		return WireframeMesh->GetBounds().TransformBy( LocalToWorld );
+	}
+
+	return FBoxSphereBounds( LocalToWorld.GetLocation(), FVector::ZeroVector, 0.0f );
+}
diff --git a/Engine/Plugins/Editor/MeshEditor/Source/MeshEditor/WireframeMeshComponent.h b/Engine/Plugins/Editor/MeshEditor/Source/MeshEditor/WireframeMeshComponent.h
new file mode 100644
--- /dev/null
+++ b/Engine/Plugins/Editor/MeshEditor/Source/MeshEditor/WireframeMeshComponent.h
@@ -0,0 +1,201 @@
+// Copyright 1998-2017 Epic Games, Inc. All Rights Reserved.
+
+#pragma once
+
+#include "CoreMinimal.h"
+#include "UObject/ObjectMacros.h"
+#include "Components/MeshComponent.h"
+#include "Containers/SparseArray.h"
+#include "RenderResource.h"
+#include "PackedNormal.h"
+#include "EditableMeshTypes.h"
+#include "WireframeMeshComponent.generated.h"
+
+class FPrimitiveSceneProxy;
+
+
+/** Structure representing a vertex definition for the wireframe mesh */
+struct FWireframeMeshVertex
+{
+	FWireframeMeshVertex() {}
+
+	FWireframeMeshVertex( const FVector& InPosition, const FVector& InNormal, const FVector& InTangent, const FColor& InColor )
+		: Position( InPosition ),
+		  TangentX( InNormal ),
+		  TangentZ( InTangent ),
+		  Color( InColor )
+	{
+		TangentZ.Vector.W = 255;
+	}
+
+	FVector Position;
+	FPackedNormal TangentX;
+	FPackedNormal TangentZ;
+	FColor Color;
+};
+
+
+/** Vertex buffer for the shared mesh */
+class FWireframeMeshVertexBuffer : public FVertexBuffer 
+{
+public:
+	TArray<FWireframeMeshVertex> Vertices;
+
+	virtual void InitRHI() override;
+};
+
+
+/** Index buffer for the shared mesh */
+class FWireframeMeshIndexBuffer : public FIndexBuffer 
+{
+public:
+	TArray<int32> Indices;
+
+	virtual void InitRHI() override;
+};
+
+
+/** Structure representing instance-specific extra data for a vertex, passed in the UV0 channel */
+struct FWireframeMeshInstanceVertex
+{
+	FWireframeMeshInstanceVertex() {}
+
+	FWireframeMeshInstanceVertex( const FVector2D& InUV )
+		: UV( InUV )
+	{}
+
+	FVector2D UV;
+};
+
+
+/** Vertex buffer for a mesh instance */
+class FWireframeMeshInstanceVertexBuffer : public FVertexBuffer
+{
+public:
+	TArray<FWireframeMeshInstanceVertex> Vertices;
+
+	virtual void InitRHI() override;
+};
+
+
+struct FWireframeVertex
+{
+	FVector Position;
+};
+
+
+struct FWireframePolygon
+{
+	FVector PolygonNormal;
+};
+
+
+struct FWireframeEdgeInstance
+{
+	FPolygonID PolygonID;
+	FEdgeID EdgeID;
+};
+
+
+struct FWireframeEdge
+{
+	FVertexID StartVertex;
+	FVertexID EndVertex;
+	FColor Color;
+	TArray<int32> EdgeInstances;
+};
+
+
+UCLASS()
+class UWireframeMesh : public UObject
+{
+	GENERATED_BODY()
+
+public:
+
+	UWireframeMesh()
+		: bBoundsDirty( true )
+	{}
+
+	void Reset();
+	void AddVertex( const FVertexID VertexID );
+	void SetVertexPosition( const FVertexID VertexID, const FVector& Position );
+	void RemoveVertex( const FVertexID VertexID );
+	void AddPolygon( const FPolygonID PolygonID );
+	void SetPolygonNormal( const FPolygonID PolygonID, const FVector& Normal );
+	void RemovePolygon( const FPolygonID PolygonID );
+	void AddEdge( const FEdgeID EdgeID );
+	void SetEdgeVertices( const FEdgeID EdgeID, const FVertexID Vertex0, const FVertexID Vertex1 );
+	void SetEdgeColor( const FEdgeID EdgeID, const FColor Color );
+	void RemoveEdge( const FEdgeID EdgeID );
+	void AddEdgeInstance( const FEdgeID EdgeID, const FPolygonID PolygonID );
+	void RemoveEdgeInstance( const FEdgeID EdgeID, const FPolygonID PolygonID );
+	const TArray<int32>& GetEdgeInstanceIDs( const FEdgeID EdgeID ) const;
+	int32 GetNumEdgeInstances() const { return EdgeInstances.Num(); }
+	FBoxSphereBounds GetBounds() const;
+
+	void InitResources();
+	void ReleaseResources();
+
+	// These arrays mirror the editable mesh elements
+	TSparseArray<FWireframeVertex> Vertices;
+	TSparseArray<FWireframePolygon> Polygons;
+	TSparseArray<FWireframeEdge> Edges;
+
+	// This is a packed array of edge instances with no holes.
+	// An edge instance represents a unique quadrilaterial which forms part of the wireframe mesh.
+	TArray<FWireframeEdgeInstance> EdgeInstances;
+
+	FWireframeMeshVertexBuffer VertexBuffer;
+	FWireframeMeshIndexBuffer IndexBuffer;
+
+	mutable FBoxSphereBounds Bounds;
+	mutable bool bBoundsDirty;
+};
+
+
+UCLASS()
+class UWireframeMeshComponent : public UMeshComponent
+{
+	GENERATED_BODY()
+
+public:
+	/**
+	 * Default UObject constructor.
+	 */
+	UWireframeMeshComponent( const FObjectInitializer& ObjectInitializer = FObjectInitializer::Get() );
+
+	/** Gets the common wireframe mesh object used by this component */
+	UWireframeMesh* GetWireframeMesh() const { return WireframeMesh; }
+
+	/** Sets the wireframe mesh object used by this component */
+	void SetWireframeMesh( UWireframeMesh* InWireframeMesh ) { WireframeMesh = InWireframeMesh; }
+
+	/** Sets all edges as visible */
+	void ShowAllEdges();
+
+	/** Changes the visibility of the named edge */
+	void SetEdgeVisibility( const FEdgeID EdgeID, const bool bEdgeVisible );
+
+private:
+
+	//~ Begin UPrimitiveComponent Interface.
+	virtual FPrimitiveSceneProxy* CreateSceneProxy() override;
+	//~ End UPrimitiveComponent Interface.
+
+	//~ Begin UMeshComponent Interface.
+	virtual int32 GetNumMaterials() const override;
+	//~ End UMeshComponent Interface.
+
+	//~ Begin USceneComponent Interface.
+	virtual FBoxSphereBounds CalcBounds( const FTransform& LocalToWorld ) const override;
+	//~ Begin USceneComponent Interface.
+
+	UPROPERTY()
+	UWireframeMesh* WireframeMesh;
+
+	/** Set of edge IDs not to be rendered for this instance */
+	TSet<FEdgeID> HiddenEdgeIDs;
+
+	friend class FWireframeMeshSceneProxy;
+};
diff --git a/Engine/Plugins/Runtime/EditableMesh/Source/EditableMesh/EditableMesh.cpp b/Engine/Plugins/Runtime/EditableMesh/Source/EditableMesh/EditableMesh.cpp
--- a/Engine/Plugins/Runtime/EditableMesh/Source/EditableMesh/EditableMesh.cpp
+++ b/Engine/Plugins/Runtime/EditableMesh/Source/EditableMesh/EditableMesh.cpp
@@ -322,6 +322,7 @@ void UEditableMesh::Compact()
 	ElementIDsRemappedEvent.Broadcast( this, Remappings );
 
 	RebuildRenderMesh();
+	RebuildOctree();
 
 	// Prepare the inverse transaction to reverse the compaction
 	FUncompactChangeInput UncompactChangeInput;
@@ -354,6 +355,7 @@ void UEditableMesh::Uncompact( const FElementIDRemappings& Remappings )
 	ElementIDsRemappedEvent.Broadcast( this, Remappings );
 
 	RebuildRenderMesh();
+	RebuildOctree();
 
 	AddUndo( MakeUnique<FCompactChange>() );
 }
@@ -7429,19 +7431,23 @@ void UEditableMesh::SetEdgesVertices( const TArray<FVerticesForEdge>& VerticesFo
 	UE_LOG( LogEditableMesh, Verbose, TEXT( "SetEdgesVertices: %s" ), *LogHelpers::ArrayToString( VerticesForEdges ) );
 
 	FSetEdgesVerticesChangeInput RevertInput;
-
 	RevertInput.VerticesForEdges.AddUninitialized( VerticesForEdges.Num() );
-	for( int32 VertexNumber = 0; VertexNumber < VerticesForEdges.Num(); ++VertexNumber )
+
+	static TArray<FEdgeID> EdgeIDs;
+	EdgeIDs.SetNumUninitialized( VerticesForEdges.Num() );
+
+	for( int32 EdgeNumber = 0; EdgeNumber < VerticesForEdges.Num(); ++EdgeNumber )
 	{
-		const FVerticesForEdge& VerticesForEdge = VerticesForEdges[ VertexNumber ];
+		const FVerticesForEdge& VerticesForEdge = VerticesForEdges[ EdgeNumber ];
 
 		// Save the backup
-		FVerticesForEdge& RevertVerticesForEdge = RevertInput.VerticesForEdges[ VertexNumber ];
+		FVerticesForEdge& RevertVerticesForEdge = RevertInput.VerticesForEdges[ EdgeNumber ];
 		RevertVerticesForEdge.EdgeID = VerticesForEdge.EdgeID;
 		GetEdgeVertices( VerticesForEdge.EdgeID, /* Out */ RevertVerticesForEdge.NewVertexID0, /* Out */ RevertVerticesForEdge.NewVertexID1 );
 
 		// Edit the edge
 		FMeshEdge& Edge = Edges[ VerticesForEdge.EdgeID.GetValue() ];
+		EdgeIDs[ EdgeNumber ] = VerticesForEdge.EdgeID;
 
 		for( uint32 EdgeVertexNumber = 0; EdgeVertexNumber < 2; ++EdgeVertexNumber )
 		{
@@ -7465,6 +7471,12 @@ void UEditableMesh::SetEdgesVertices( const TArray<FVerticesForEdge>& VerticesFo
 		}
 	}
 
+	// Give the adapter a chance to handle this
+	for( UEditableMeshAdapter* Adapter : Adapters )
+	{
+		Adapter->OnSetEdgesVertices( this, EdgeIDs );
+	}
+
 	AddUndo( MakeUnique<FSetEdgesVerticesChange>( MoveTemp( RevertInput ) ) );
 
 	UE_LOG( LogEditableMesh, Verbose, TEXT( "* SetEdgesVertices finished") );
diff --git a/Engine/Plugins/Runtime/EditableMesh/Source/EditableMesh/EditableStaticMeshAdapter.cpp b/Engine/Plugins/Runtime/EditableMesh/Source/EditableMesh/EditableStaticMeshAdapter.cpp
--- a/Engine/Plugins/Runtime/EditableMesh/Source/EditableMesh/EditableStaticMeshAdapter.cpp
+++ b/Engine/Plugins/Runtime/EditableMesh/Source/EditableMesh/EditableStaticMeshAdapter.cpp
@@ -1496,6 +1496,12 @@ void UEditableStaticMeshAdapter::OnDeleteEdges( const UEditableMesh* EditableMes
 }
 
 
+void UEditableStaticMeshAdapter::OnSetEdgesVertices( const UEditableMesh* EditableMesh, const TArray<FEdgeID>& EdgeIDs )
+{
+	// Nothing to do here
+}
+
+
 void UEditableStaticMeshAdapter::OnDeletePolygons( const UEditableMesh* EditableMesh, const TArray<FPolygonID>& PolygonIDs )
 {
 	for( const FPolygonID PolygonID : PolygonIDs )
diff --git a/Engine/Plugins/Runtime/EditableMesh/Source/EditableMesh/EditableStaticMeshAdapter.h b/Engine/Plugins/Runtime/EditableMesh/Source/EditableMesh/EditableStaticMeshAdapter.h
--- a/Engine/Plugins/Runtime/EditableMesh/Source/EditableMesh/EditableStaticMeshAdapter.h
+++ b/Engine/Plugins/Runtime/EditableMesh/Source/EditableMesh/EditableStaticMeshAdapter.h
@@ -143,6 +143,7 @@ public:
 	virtual void OnSetVertexInstanceAttribute( const UEditableMesh* EditableMesh, const FVertexInstanceID VertexInstanceID, const FName AttributeName, const int32 AttributeIndex, const FVector4 AttributeValue ) override;
 	virtual void OnCreateEdges( const UEditableMesh* EditableMesh, const TArray<FEdgeID>& EdgeIDs ) override;
 	virtual void OnDeleteEdges( const UEditableMesh* EditableMesh, const TArray<FEdgeID>& EdgeIDs ) override;
+	virtual void OnSetEdgesVertices( const UEditableMesh* EditableMesh, const TArray<FEdgeID>& EdgeIDs ) override;
 	virtual void OnSetEdgeAttribute( const UEditableMesh* EditableMesh, const FEdgeID EdgeID, const FName AttributeName, const int32 AttributeIndex, const FVector4 AttributeValue ) override;
 	virtual void OnCreatePolygons( const UEditableMesh* EditableMesh, const TArray<FPolygonID>& PolygonIDs ) override;
 	virtual void OnDeletePolygons( const UEditableMesh* EditableMesh, const TArray<FPolygonID>& PolygonIDs ) override;
diff --git a/Engine/Plugins/Runtime/EditableMesh/Source/EditableMesh/Public/EditableMesh.h b/Engine/Plugins/Runtime/EditableMesh/Source/EditableMesh/Public/EditableMesh.h
--- a/Engine/Plugins/Runtime/EditableMesh/Source/EditableMesh/Public/EditableMesh.h
+++ b/Engine/Plugins/Runtime/EditableMesh/Source/EditableMesh/Public/EditableMesh.h
@@ -908,7 +908,7 @@ public:
 	int32 PendingCompactCounter;
 
 	/** Data will be compacted after this many topology modifying actions. */
-	static const int32 CompactFrequency = 10;
+	static const int32 CompactFrequency = 5;
 
 	/** OpenSubdiv topology refiner object.  This is generated for meshes that have subdivision levels, and reused to generate new limit surfaces 
 	    when geometry is moved.  When the mesh's topology changes, this object is regenerated from scratch. */
diff --git a/Engine/Plugins/Runtime/EditableMesh/Source/EditableMesh/Public/EditableMeshAdapter.h b/Engine/Plugins/Runtime/EditableMesh/Source/EditableMesh/Public/EditableMeshAdapter.h
--- a/Engine/Plugins/Runtime/EditableMesh/Source/EditableMesh/Public/EditableMeshAdapter.h
+++ b/Engine/Plugins/Runtime/EditableMesh/Source/EditableMesh/Public/EditableMeshAdapter.h
@@ -10,7 +10,7 @@
 
 
 UCLASS(Abstract)
-class UEditableMeshAdapter : public UObject
+class EDITABLEMESH_API UEditableMeshAdapter : public UObject
 {
 	GENERATED_BODY()
 
@@ -39,6 +39,7 @@ public:
 	virtual void OnSetVertexInstanceAttribute( const UEditableMesh* EditableMesh, const FVertexInstanceID VertexInstanceID, const FName AttributeName, const int32 AttributeIndex, const FVector4 AttributeValue ) PURE_VIRTUAL(,);
 	virtual void OnCreateEdges( const UEditableMesh* EditableMesh, const TArray<FEdgeID>& EdgeIDs ) PURE_VIRTUAL(,);
 	virtual void OnDeleteEdges( const UEditableMesh* EditableMesh, const TArray<FEdgeID>& EdgeIDs ) PURE_VIRTUAL(,);
+	virtual void OnSetEdgesVertices( const UEditableMesh* EditableMesh, const TArray<FEdgeID>& EdgeIDs ) PURE_VIRTUAL(,);
 	virtual void OnSetEdgeAttribute( const UEditableMesh* EditableMesh, const FEdgeID EdgeID, const FName AttributeName, const int32 AttributeIndex, const FVector4 AttributeValue ) PURE_VIRTUAL(,);
 	virtual void OnCreatePolygons( const UEditableMesh* EditableMesh, const TArray<FPolygonID>& PolygonIDs ) PURE_VIRTUAL(,);
 	virtual void OnDeletePolygons( const UEditableMesh* EditableMesh, const TArray<FPolygonID>& PolygonIDs ) PURE_VIRTUAL(,);
