Fixed unity builds.

diff --git a/Engine/Plugins/Editor/MeshEditor/Source/MeshEditor/MeshEditorSubdividedStaticMeshAdapter.cpp b/Engine/Plugins/Editor/MeshEditor/Source/MeshEditor/MeshEditorSubdividedStaticMeshAdapter.cpp
--- a/Engine/Plugins/Editor/MeshEditor/Source/MeshEditor/MeshEditorSubdividedStaticMeshAdapter.cpp
+++ b/Engine/Plugins/Editor/MeshEditor/Source/MeshEditor/MeshEditorSubdividedStaticMeshAdapter.cpp
@@ -5,7 +5,7 @@
 #include "WireframeMeshComponent.h"
 
 
-static FColor GetEdgeColor( bool bIsSubdividedWireEdge )
+static FColor GetSubdividedEdgeColor( bool bIsSubdividedWireEdge )
 {
 	return bIsSubdividedWireEdge ? FColor( 0, 0, 51, 204 ) : FColor( 13, 13, 13, 153 );
 }
@@ -72,7 +72,7 @@ void UMeshEditorSubdividedStaticMeshAdapter::OnRebuildRenderMesh( const UEditabl
 			const FEdgeID EdgeID( EdgeIndex );
 			WireframeMesh->AddEdge( EdgeID );
 			WireframeMesh->SetEdgeVertices( EdgeID, FVertexID( SubdividedWireEdge.EdgeVertex0PositionIndex ), FVertexID( SubdividedWireEdge.EdgeVertex1PositionIndex ) );
-			WireframeMesh->SetEdgeColor( EdgeID, GetEdgeColor( SubdividedWireEdge.CounterpartEdgeID != FEdgeID::Invalid ) );
+			WireframeMesh->SetEdgeColor( EdgeID, GetSubdividedEdgeColor( SubdividedWireEdge.CounterpartEdgeID != FEdgeID::Invalid ) );
 			WireframeMesh->AddEdgeInstance( EdgeID, DummyPolygonID );
 			EdgeIndex++;
 		}
