Submitted missing build file.

diff --git a/Engine/Plugins/Editor/MeshEditor/Source/MeshEditor/MeshEditor.Build.cs b/Engine/Plugins/Editor/MeshEditor/Source/MeshEditor/MeshEditor.Build.cs
--- a/Engine/Plugins/Editor/MeshEditor/Source/MeshEditor/MeshEditor.Build.cs
+++ b/Engine/Plugins/Editor/MeshEditor/Source/MeshEditor/MeshEditor.Build.cs
@@ -20,9 +20,11 @@ namespace UnrealBuildTool.Rules
 					"EditableMesh",
 					"ViewportInteraction",
                     "VREditor",
-					"Projects"
-				}
-			);
+					"Projects",
+                    "ShaderCore",
+                    "RHI"
+                }
+            );
 
 			PrivateIncludePathModuleNames.AddRange(
 				new string[] {
@@ -37,6 +39,6 @@ namespace UnrealBuildTool.Rules
 					"LevelEditor",
 				}
 			);
-		}
-	}
+        }
+    }
 }
\ No newline at end of file
