- Split off mesh description into its own class and engine module. The intention is this will be shared by other mesh objects as the source mesh format, rather than being part of UEditableMesh.
- Created new API for accessing mesh elements. They now allow a range-for iteration over element IDs, and a custom container type for elements.
- Beginnings of a new generic attribute system which supports different attribute value types. This provides SoA type attribute containers (so, for example, all vertex positions are held contiguously).
- Revisited polygon groups and materials. Materials are now handled like an attribute of polygon groups, and are represented by a string asset reference, and two slot names (user-given and imported). Added a new basic action to move a polygon from one group to another, and removed the 'assign material to polygon' action. A mesh can now have more than one polygon group with the same material.
- General refactoring and optimization.

diff --git a/Engine/Plugins/Editor/MeshEditor/Source/MeshEditor/MeshEditor.Build.cs b/Engine/Plugins/Editor/MeshEditor/Source/MeshEditor/MeshEditor.Build.cs
--- a/Engine/Plugins/Editor/MeshEditor/Source/MeshEditor/MeshEditor.Build.cs
+++ b/Engine/Plugins/Editor/MeshEditor/Source/MeshEditor/MeshEditor.Build.cs
@@ -18,6 +18,7 @@ namespace UnrealBuildTool.Rules
 				    "UnrealEd",
 					"RenderCore",
 					"EditableMesh",
+                    "MeshDescription",
 					"ViewportInteraction",
                     "VREditor",
 					"Projects",
@@ -30,6 +31,7 @@ namespace UnrealBuildTool.Rules
 				new string[] {
                     "ContentBrowser",
 					"LevelEditor",
+                    "MeshDescription"
 				}
 			);
 
diff --git a/Engine/Plugins/Editor/MeshEditor/Source/MeshEditor/MeshEditorMode.cpp b/Engine/Plugins/Editor/MeshEditor/Source/MeshEditor/MeshEditorMode.cpp
--- a/Engine/Plugins/Editor/MeshEditor/Source/MeshEditor/MeshEditorMode.cpp
+++ b/Engine/Plugins/Editor/MeshEditor/Source/MeshEditor/MeshEditorMode.cpp
@@ -1285,29 +1285,24 @@ void FMeshEditorMode::UpdateDebugNormals()
 
 		const FMatrix ComponentToWorldMatrix = Component->GetRenderMatrix();
 
-		const uint32 PolygonArraySize = EditableMesh->GetPolygonArraySize();
-		for( uint32 PolygonIndex = 0; PolygonIndex < PolygonArraySize; ++PolygonIndex )
+		for( const FPolygonID PolygonID : EditableMesh->GetMeshDescription()->Polygons().GetElementIDs() )
 		{
-			const FPolygonID PolygonID( PolygonIndex );
-			if( EditableMesh->IsValidPolygon( PolygonID ) )
-			{
-				// @todo mesheditor: total debug feature for now. Need a way of making this look nice.
-				const int32 PerimeterVertexCount = EditableMesh->GetPolygonPerimeterVertexCount( PolygonID );
-				const float Length = 10.0f; // @todo mesheditor: determine length of debug line from distance from the mesh origin to the camera?
+			// @todo mesheditor: total debug feature for now. Need a way of making this look nice.
+			const int32 PerimeterVertexCount = EditableMesh->GetPolygonPerimeterVertexCount( PolygonID );
+			const float Length = 10.0f; // @todo mesheditor: determine length of debug line from distance from the mesh origin to the camera?
 
-				for( int32 PerimeterVertexIndex = 0; PerimeterVertexIndex < PerimeterVertexCount; ++PerimeterVertexIndex )
-				{
-					const FVector Position( EditableMesh->GetPolygonPerimeterVertexAttribute( PolygonID, PerimeterVertexIndex, UEditableMeshAttribute::VertexPosition(), 0 ) );
-					const FVector Normal( EditableMesh->GetPolygonPerimeterVertexAttribute( PolygonID, PerimeterVertexIndex, UEditableMeshAttribute::VertexNormal(), 0 ) );
-					const FVector Tangent( EditableMesh->GetPolygonPerimeterVertexAttribute( PolygonID, PerimeterVertexIndex, UEditableMeshAttribute::VertexTangent(), 0 ) );
+			for( int32 PerimeterVertexIndex = 0; PerimeterVertexIndex < PerimeterVertexCount; ++PerimeterVertexIndex )
+			{
+				const FVector Position( EditableMesh->GetPolygonPerimeterVertexAttribute( PolygonID, PerimeterVertexIndex, UEditableMeshAttribute::VertexPosition(), 0 ) );
+				const FVector Normal( EditableMesh->GetPolygonPerimeterVertexAttribute( PolygonID, PerimeterVertexIndex, UEditableMeshAttribute::VertexNormal(), 0 ) );
+				const FVector Tangent( EditableMesh->GetPolygonPerimeterVertexAttribute( PolygonID, PerimeterVertexIndex, UEditableMeshAttribute::VertexTangent(), 0 ) );
 
-					const FVector Start( ComponentToWorldMatrix.TransformPosition( Position ) );
-					const FVector NormalEnd( ComponentToWorldMatrix.TransformPosition( Position + Normal * Length ) );
-					const FVector TangentEnd( ComponentToWorldMatrix.TransformPosition( Position + Tangent * Length * 0.5f ) );
+				const FVector Start( ComponentToWorldMatrix.TransformPosition( Position ) );
+				const FVector NormalEnd( ComponentToWorldMatrix.TransformPosition( Position + Normal * Length ) );
+				const FVector TangentEnd( ComponentToWorldMatrix.TransformPosition( Position + Tangent * Length * 0.5f ) );
 
-					DebugNormalsComponent->AddLine( FOverlayLine( Start, NormalEnd, FColor::Magenta, 0.0f ) );
-					DebugNormalsComponent->AddLine( FOverlayLine( Start, TangentEnd, FColor::Yellow, 0.0f ) );
-				}
+				DebugNormalsComponent->AddLine( FOverlayLine( Start, NormalEnd, FColor::Magenta, 0.0f ) );
+				DebugNormalsComponent->AddLine( FOverlayLine( Start, TangentEnd, FColor::Yellow, 0.0f ) );
 			}
 		}
 	}
@@ -2246,12 +2241,6 @@ bool FMeshEditorMode::AssignMaterialToSelectedPolygons( UMaterialInterface* Sele
 			return false;
 		}
 
-		static TArray<FPolygonID> NewPolygonIDs;
-		NewPolygonIDs.Reset();
-
-		static TArray<FMeshElement> MeshElementsToSelect;
-		MeshElementsToSelect.Reset();
-
 		const FScopedTransaction Transaction( LOCTEXT( "UndoAssignMaterialToPolygon", "Assign Material to Polygon" ) );
 
 		CommitSelectedMeshes();
@@ -2259,10 +2248,6 @@ bool FMeshEditorMode::AssignMaterialToSelectedPolygons( UMaterialInterface* Sele
 		// Refresh selection (committing may have created a new mesh instance)
 		GetSelectedMeshesAndPolygons( /* Out */ MeshesAndPolygons );
 
-		// Deselect the polygons first.  They'll be moved to a new section, and hence be assigned new PolygonRefs.  This will also make sure 
-		// they become selected again after undo.
-		DeselectMeshElements( MeshesAndPolygons );
-
 		for( const auto& MeshAndPolygons : MeshesAndPolygons )
 		{
 			UEditableMesh* EditableMesh = MeshAndPolygons.Key;
@@ -2277,39 +2262,78 @@ bool FMeshEditorMode::AssignMaterialToSelectedPolygons( UMaterialInterface* Sele
 
 			EditableMesh->StartModification( EMeshModificationType::Final, EMeshTopologyChange::TopologyChange );
 			{
-				static TArray<FPolygonID> PolygonsToAssign;
-				PolygonsToAssign.Reset();
+				// See if there's a polygon group using this material, and if not create one.
+				// @todo mesheditor: This currently imposes the limitation that each polygon group has a unique material.
+				// Eventually we will need to be able to specify PolygonGroup properties in the editor, and ask the user for
+				// further details if there is more than one polygon group which matches the material.
+				FPolygonGroupID PolygonGroupToAssign = FPolygonGroupID::Invalid;
+				for( const FPolygonGroupID PolygonGroupID : EditableMesh->GetMeshDescription()->PolygonGroups().GetElementIDs() )
+				{
+					UMaterialInterface* Material = EditableMesh->GetPolygonGroupMaterial( PolygonGroupID );
+					if( Material == SelectedMaterial )
+					{
+						// We only expect to find one polygon group containing this material at the moment.
+						// We need to provide a way of distinguishing different polygon groups with the same material.
+						ensure( PolygonGroupToAssign == FPolygonGroupID::Invalid );
 
-				for( const FMeshElement& PolygonElement : MeshAndPolygons.Value )
+						PolygonGroupToAssign = PolygonGroupID;
+					}
+				}
+
+				// If we didn't find the material being used anywhere, create a new polygon group
+				if( PolygonGroupToAssign == FPolygonGroupID::Invalid )
 				{
-					const FPolygonID PolygonID( PolygonElement.ElementAddress.ElementID );
-					PolygonsToAssign.Add( PolygonID );
+					// Helper function which returns a unique FName for the material slot name, based on the material's asset name,
+					// and adding a unique suffix if there are other polygon groups with the same material slot name.
+					auto MakeUniqueSlotName = []( UEditableMesh* EditableMesh, FName Name ) -> FName
+					{
+						for( const FPolygonGroupID PolygonGroupID : EditableMesh->GetMeshDescription()->PolygonGroups().GetElementIDs() )
+						{
+							const FName ExistingName = EditableMesh->GetPolygonGroupMaterialSlotName( PolygonGroupID );
+							if( ExistingName.GetComparisonIndex() == Name.GetComparisonIndex() )
+							{
+								Name = FName( Name, FMath::Max( Name.GetNumber(), ExistingName.GetNumber() + 1 ) );
+							}
+						}
+						return Name;
+					};
+
+					static TArray<FPolygonGroupToCreate> PolygonGroupsToCreate;
+					PolygonGroupsToCreate.Reset( 1 );
+					PolygonGroupsToCreate.Emplace();
+					
+					FPolygonGroupToCreate& PolygonGroupToCreate = PolygonGroupsToCreate.Last();
+					PolygonGroupToCreate.MaterialAsset = SelectedMaterial->GetPathName();
+					PolygonGroupToCreate.MaterialSlotName = MakeUniqueSlotName( EditableMesh, SelectedMaterial->GetFName() );
+					PolygonGroupToCreate.bEnableCollision = true;
+					PolygonGroupToCreate.bCastShadow = true;
+
+					static TArray<FPolygonGroupID> NewPolygonGroupIDs;
+					EditableMesh->CreatePolygonGroups( PolygonGroupsToCreate, NewPolygonGroupIDs );
+					PolygonGroupToAssign = NewPolygonGroupIDs[ 0 ];
 				}
 
-				EditableMesh->AssignMaterialToPolygons( PolygonsToAssign, SelectedMaterial, NewPolygonIDs );
+				static TArray<FPolygonGroupForPolygon> PolygonsToAssign;
+				PolygonsToAssign.Reset();
 
-				for( const FPolygonID NewPolygonID : NewPolygonIDs )
+				for( const FMeshElement& PolygonElement : MeshAndPolygons.Value )
 				{
-					// Select the new polygon
-					FMeshElement NewPolygonMeshElement;
-					{
-						NewPolygonMeshElement.Component = Component;
-						NewPolygonMeshElement.ElementAddress.SubMeshAddress = EditableMesh->GetSubMeshAddress();
-						NewPolygonMeshElement.ElementAddress.ElementType = EEditableMeshElementType::Polygon;
-						NewPolygonMeshElement.ElementAddress.ElementID = NewPolygonID;
-					}
+					const FPolygonID PolygonID( PolygonElement.ElementAddress.ElementID );
 
-					MeshElementsToSelect.Add( NewPolygonMeshElement );
+					PolygonsToAssign.Emplace();
+					FPolygonGroupForPolygon& PolygonGroupForPolygon = PolygonsToAssign.Last();
+					PolygonGroupForPolygon.PolygonID = PolygonID;
+					PolygonGroupForPolygon.PolygonGroupID = PolygonGroupToAssign;
 				}
+
+				const bool bDeleteOrphanedPolygonGroups = true;
+				EditableMesh->AssignPolygonsToPolygonGroups( PolygonsToAssign, bDeleteOrphanedPolygonGroups );
 			}
 			EditableMesh->EndModification();
 
 			TrackUndo( EditableMesh, EditableMesh->MakeUndo() );
 		}
 
-		// Select the newly-created triangles
-		SelectMeshElements( MeshElementsToSelect );
-
 		return true;
 	}
 
@@ -2421,7 +2445,9 @@ void FMeshEditorMode::AddMeshElementToOverlay( UOverlayComponent* OverlayCompone
 							FVector VertexPositions[ 3 ];
 							for( int32 TriangleVertex = 0; TriangleVertex < 3; TriangleVertex++ )
 							{
-								VertexPositions[ TriangleVertex ] = ComponentToWorldMatrix.TransformPosition( EditableMesh->GetPolygonTriangulatedTriangleVertexPosition( PolygonID, PolygonTriangle, TriangleVertex ) );
+								const FVertexInstanceID VertexInstanceID = EditableMesh->GetPolygonTriangulatedTriangle( PolygonID, PolygonTriangle ).GetVertexInstanceID( TriangleVertex );
+								const FVertexID VertexID = EditableMesh->GetVertexInstanceVertex( VertexInstanceID );
+								VertexPositions[ TriangleVertex ] = ComponentToWorldMatrix.TransformPosition( EditableMesh->GetVertexAttribute( VertexID, UEditableMeshAttribute::VertexPosition(), 0 ) );
 							}
 
 							OverlayComponent->AddTriangle( FOverlayTriangle( 
@@ -2454,15 +2480,15 @@ void FMeshEditorMode::Render( const FSceneView* SceneView, FViewport* Viewport,
 			GEngine->AddOnScreenDebugMessage( INDEX_NONE, 0.0f, FColor::White,
 				FString::Printf( TEXT( "MeshElement: %s" ), *MeshElement.ToString() ), false );
 			GEngine->AddOnScreenDebugMessage( INDEX_NONE, 0.0f, FColor::White,
-				FString::Printf( TEXT( "Vertices: %i (array size: %i)" ), EditableMesh->GetVertexCount(), EditableMesh->GetVertexArraySize() ), false );
+				FString::Printf( TEXT( "Vertices: %i (array size: %i)" ), EditableMesh->GetVertexCount(), EditableMesh->GetMeshDescription()->Vertices().GetArraySize() ), false );
 			GEngine->AddOnScreenDebugMessage( INDEX_NONE, 0.0f, FColor::White,
-				FString::Printf( TEXT( "Vertex instances: %i (array size: %i)" ), EditableMesh->GetVertexInstanceCount(), EditableMesh->GetVertexInstanceArraySize() ), false );
+				FString::Printf( TEXT( "Vertex instances: %i (array size: %i)" ), EditableMesh->GetVertexInstanceCount(), EditableMesh->GetMeshDescription()->VertexInstances().GetArraySize() ), false );
 			GEngine->AddOnScreenDebugMessage( INDEX_NONE, 0.0f, FColor::White,
-				FString::Printf( TEXT( "Edges: %i (array size: %i)" ), EditableMesh->GetEdgeCount(), EditableMesh->GetEdgeArraySize() ), false );
+				FString::Printf( TEXT( "Edges: %i (array size: %i)" ), EditableMesh->GetEdgeCount(), EditableMesh->GetMeshDescription()->Edges().GetArraySize() ), false );
 			GEngine->AddOnScreenDebugMessage( INDEX_NONE, 0.0f, FColor::White,
-				FString::Printf( TEXT( "Polygons: %i (array size: %i)" ), EditableMesh->GetPolygonCount(), EditableMesh->GetPolygonArraySize() ), false );
+				FString::Printf( TEXT( "Polygons: %i (array size: %i)" ), EditableMesh->GetPolygonCount(), EditableMesh->GetMeshDescription()->Polygons().GetArraySize() ), false );
 			GEngine->AddOnScreenDebugMessage( INDEX_NONE, 0.0f, FColor::White,
-				FString::Printf( TEXT( "Sections: %i (array size: %i)" ), EditableMesh->GetPolygonGroupCount(), EditableMesh->GetPolygonGroupArraySize() ), false );
+				FString::Printf( TEXT( "Sections: %i (array size: %i)" ), EditableMesh->GetPolygonGroupCount(), EditableMesh->GetMeshDescription()->PolygonGroups().GetArraySize() ), false );
 			// @todo mesheditor: triangles array is now an implementation detail in the adapter. Discuss if it's worth preserving access to it somehow.
 			//for( int32 PolygonGroupIndex = 0; PolygonGroupIndex < EditableMesh->GetPolygonGroupArraySize(); ++PolygonGroupIndex )
 			//{
@@ -3496,15 +3522,9 @@ FEditableMeshElementAddress FMeshEditorMode::QueryElement( const UEditableMesh&
 	else
 	{
 		// @todo mesheditor spatial: Need GrabberSphere support for spatial queries.  Currently we're just testing all polygons (slow!)
-		const uint32 PolygonArraySize = EditableMesh.GetPolygonArraySize();
-		for( uint32 PolygonIndex = 0; PolygonIndex < PolygonArraySize; ++PolygonIndex )
+		for( const FPolygonID PolygonID : EditableMesh.GetMeshDescription()->Polygons().GetElementIDs() )
 		{
-			const FPolygonID PolygonID( PolygonIndex );
-
-			if( EditableMesh.IsValidPolygon( PolygonID ) )
-			{
-				CandidatePolygons.Add( PolygonID );
-			}
+			CandidatePolygons.Add( PolygonID );
 		}
 	}
 
@@ -3557,8 +3577,9 @@ FEditableMeshElementAddress FMeshEditorMode::QueryElement( const UEditableMesh&
 			FVector TriangleVertexPositions[ 3 ];
 			for( uint32 TriangleVertexNumber = 0; TriangleVertexNumber < 3; ++TriangleVertexNumber )
 			{
-				TriangleVertexPositions[ TriangleVertexNumber ] =
-					EditableMesh.GetPolygonTriangulatedTriangleVertexPosition( PolygonID, PolygonTriangleNumber, TriangleVertexNumber );
+				const FVertexInstanceID VertexInstanceID = EditableMesh.GetPolygonTriangulatedTriangle( PolygonID, PolygonTriangleNumber ).GetVertexInstanceID( TriangleVertexNumber );
+				const FVertexID VertexID = EditableMesh.GetVertexInstanceVertex( VertexInstanceID );
+				TriangleVertexPositions[ TriangleVertexNumber ] = EditableMesh.GetVertexAttribute( VertexID, UEditableMeshAttribute::VertexPosition(), 0 );
 			}
 
 			const bool bAlreadyHitTriangle = ( HitElementAddress.ElementType == EEditableMeshElementType::Polygon );
@@ -4352,82 +4373,72 @@ bool FMeshEditorMode::FrustumSelect( const FConvexVolume& InFrustum, FEditorView
 		SelectedVertexIDs.Empty();
 
 		// First, find all edges which lie at least partially within the frustum.
-		const int32 EdgeArraySize = EditableMesh->GetEdgeArraySize();
-		for( int32 EdgeIndex = 0; EdgeIndex < EdgeArraySize; ++EdgeIndex )
+		for( const FEdgeID EdgeID : EditableMesh->GetMeshDescription()->Edges().GetElementIDs() )
 		{
-			const FEdgeID EdgeID( EdgeIndex );
-			if( EditableMesh->IsValidEdge( EdgeID ) )
-			{
-				const FVertexID VertexID1( EditableMesh->GetEdgeVertex( EdgeID, 0 ) );
-				const FVertexID VertexID2( EditableMesh->GetEdgeVertex( EdgeID, 1 ) );
+			const FVertexID VertexID1( EditableMesh->GetEdgeVertex( EdgeID, 0 ) );
+			const FVertexID VertexID2( EditableMesh->GetEdgeVertex( EdgeID, 1 ) );
 
-				const FVector VertexPosition1( ComponentTransform.TransformPosition( EditableMesh->GetVertexAttribute( VertexID1, UEditableMeshAttribute::VertexPosition(), 0 ) ) );
-				const FVector VertexPosition2( ComponentTransform.TransformPosition( EditableMesh->GetVertexAttribute( VertexID2, UEditableMeshAttribute::VertexPosition(), 0 ) ) );
+			const FVector VertexPosition1( ComponentTransform.TransformPosition( EditableMesh->GetVertexAttribute( VertexID1, UEditableMeshAttribute::VertexPosition(), 0 ) ) );
+			const FVector VertexPosition2( ComponentTransform.TransformPosition( EditableMesh->GetVertexAttribute( VertexID2, UEditableMeshAttribute::VertexPosition(), 0 ) ) );
 
-				if( InFrustum.IntersectLineSegment( VertexPosition1, VertexPosition2 ) )
-				{
-					bool bAreAllPolysBackFacing = true;
+			if( InFrustum.IntersectLineSegment( VertexPosition1, VertexPosition2 ) )
+			{
+				bool bAreAllPolysBackFacing = true;
 
-					// Now iterate through all connected polygons.
-					// If any are front facing, we consider the edge also to be front facing.
-					const int32 EdgeConnectedPolygonCount = EditableMesh->GetEdgeConnectedPolygonCount( EdgeID );
-					for( int32 EdgeConnectedPolygonIndex = 0; EdgeConnectedPolygonIndex < EdgeConnectedPolygonCount; ++EdgeConnectedPolygonIndex )
-					{
-						const FPolygonID EdgeConnectedPolygonID( EditableMesh->GetEdgeConnectedPolygon( EdgeID, EdgeConnectedPolygonIndex ) );
+				// Now iterate through all connected polygons.
+				// If any are front facing, we consider the edge also to be front facing.
+				const int32 EdgeConnectedPolygonCount = EditableMesh->GetEdgeConnectedPolygonCount( EdgeID );
+				for( int32 EdgeConnectedPolygonIndex = 0; EdgeConnectedPolygonIndex < EdgeConnectedPolygonCount; ++EdgeConnectedPolygonIndex )
+				{
+					const FPolygonID EdgeConnectedPolygonID( EditableMesh->GetEdgeConnectedPolygon( EdgeID, EdgeConnectedPolygonIndex ) );
 
-						// Determine whether polygon is back facing or not using dot product of its normal with the direction vector from the eye position to somewhere on the plane
-						// (one of the vertex positions is sufficient for this)
-						const FVector PolyNormal = ComponentTransform.TransformVector( EditableMesh->ComputePolygonNormal( EdgeConnectedPolygonID ) );
-						const FVector ViewDirection = VertexPosition1 - CachedCameraToWorld.GetValue().GetLocation();
-						const bool bIsBackFacing = ( FVector::DotProduct( PolyNormal, ViewDirection ) > 0.0f );
+					// Determine whether polygon is back facing or not using dot product of its normal with the direction vector from the eye position to somewhere on the plane
+					// (one of the vertex positions is sufficient for this)
+					const FVector PolyNormal = ComponentTransform.TransformVector( EditableMesh->ComputePolygonNormal( EdgeConnectedPolygonID ) );
+					const FVector ViewDirection = VertexPosition1 - CachedCameraToWorld.GetValue().GetLocation();
+					const bool bIsBackFacing = ( FVector::DotProduct( PolyNormal, ViewDirection ) > 0.0f );
 
-						bAreAllPolysBackFacing &= bIsBackFacing;
+					bAreAllPolysBackFacing &= bIsBackFacing;
 
-						if( !bOnlySelectVisibleElements || !bIsBackFacing )
-						{
-							// Add the polygon if it is front facing, or if we don't care about only selecting visible elements
-							SelectedPolygonIDs.Add( EdgeConnectedPolygonID );
-						}
+					if( !bOnlySelectVisibleElements || !bIsBackFacing )
+					{
+						// Add the polygon if it is front facing, or if we don't care about only selecting visible elements
+						SelectedPolygonIDs.Add( EdgeConnectedPolygonID );
 					}
+				}
 
-					if( !bOnlySelectVisibleElements || !bAreAllPolysBackFacing )
-					{
-						// If at least one of the connected polygons is front facing, we deem the edge also to be front facing
-						SelectedEdgeIDs.Add( EdgeID );
+				if( !bOnlySelectVisibleElements || !bAreAllPolysBackFacing )
+				{
+					// If at least one of the connected polygons is front facing, we deem the edge also to be front facing
+					SelectedEdgeIDs.Add( EdgeID );
 
-						// Just because the edge is in the frustum doesn't imply that both its constituent vertices are.
-						// We have to do further frustum / point checks.
-						if( InFrustum.IntersectPoint( VertexPosition1 ) )
-						{
-							SelectedVertexIDs.Add( VertexID1 );
-						}
+					// Just because the edge is in the frustum doesn't imply that both its constituent vertices are.
+					// We have to do further frustum / point checks.
+					if( InFrustum.IntersectPoint( VertexPosition1 ) )
+					{
+						SelectedVertexIDs.Add( VertexID1 );
+					}
 
-						if( InFrustum.IntersectPoint( VertexPosition2 ) )
-						{
-							SelectedVertexIDs.Add( VertexID2 );
-						}
+					if( InFrustum.IntersectPoint( VertexPosition2 ) )
+					{
+						SelectedVertexIDs.Add( VertexID2 );
 					}
 				}
 			}
 		}
 
 		// Next, look for any orphaned vertices (i.e. which do not form part of an edge)
-		const int32 VertexArraySize = EditableMesh->GetVertexArraySize();
-		for( int32 VertexIndex = 0; VertexIndex < VertexArraySize; ++VertexIndex )
+		for( const FVertexID VertexID : EditableMesh->GetMeshDescription()->Vertices().GetElementIDs() )
 		{
-			const FVertexID VertexID( VertexIndex );
-			if( EditableMesh->IsValidVertex( VertexID ) )
+			// If the vertex has connected edges, it will have already been considered in the above code.
+			// Here we only want to catch vertices with no associated edges.
+			if( EditableMesh->GetVertexConnectedEdgeCount( VertexID ) == 0 )
 			{
-				// If the vertex has connected edges, it will have already been considered in the above code.
-				// Here we only want to catch vertices with no associated edges.
-				if( EditableMesh->GetVertexConnectedEdgeCount( VertexID ) == 0 )
+				const FVector VertexPosition( ComponentTransform.TransformPosition( EditableMesh->GetVertexAttribute( VertexID, UEditableMeshAttribute::VertexPosition(), 0 ) ) );
+				if( InFrustum.IntersectPoint( VertexPosition ) )
 				{
-					const FVector VertexPosition( ComponentTransform.TransformPosition( EditableMesh->GetVertexAttribute( VertexID, UEditableMeshAttribute::VertexPosition(), 0 ) ) );
-					if( InFrustum.IntersectPoint( VertexPosition ) )
-					{
-						// As the vertex is orphaned, it cannot be front or back facing. So we add it regardless.
-						SelectedVertexIDs.Add( VertexID );
-					}
+					// As the vertex is orphaned, it cannot be front or back facing. So we add it regardless.
+					SelectedVertexIDs.Add( VertexID );
 				}
 			}
 		}
diff --git a/Engine/Plugins/Editor/MeshEditor/Source/MeshEditor/MeshEditorStaticMeshAdapter.cpp b/Engine/Plugins/Editor/MeshEditor/Source/MeshEditor/MeshEditorStaticMeshAdapter.cpp
--- a/Engine/Plugins/Editor/MeshEditor/Source/MeshEditor/MeshEditorStaticMeshAdapter.cpp
+++ b/Engine/Plugins/Editor/MeshEditor/Source/MeshEditor/MeshEditorStaticMeshAdapter.cpp
@@ -48,28 +48,24 @@ void UMeshEditorStaticMeshAdapter::OnRebuildRenderMesh( const UEditableMesh* Edi
 	WireframeMesh->Reset();
 
 	// Add all vertices
-	for( TSparseArray<FMeshVertex>::TConstIterator It( EditableMesh->Vertices ); It; ++It )
+	for( const FVertexID VertexID : EditableMesh->GetMeshDescription()->Vertices().GetElementIDs() )
 	{
-		const FVertexID VertexID( It.GetIndex() );
 		WireframeMesh->AddVertex( VertexID );
 		WireframeMesh->SetVertexPosition( VertexID, EditableMesh->GetVertexAttribute( VertexID, UEditableMeshAttribute::VertexPosition(), 0 ) );
 	}
 
 	// Add all edges
-	for( TSparseArray<FMeshEdge>::TConstIterator It( EditableMesh->Edges ); It; ++It )
+	for( const FEdgeID EdgeID : EditableMesh->GetMeshDescription()->Edges().GetElementIDs() )
 	{
-		const FEdgeID EdgeID( It.GetIndex() );
 		WireframeMesh->AddEdge( EdgeID );
-		WireframeMesh->SetEdgeVertices( EdgeID, It->VertexIDs[ 0 ], It->VertexIDs[ 1 ] );
+		WireframeMesh->SetEdgeVertices( EdgeID, EditableMesh->GetEdgeVertex( EdgeID, 0 ), EditableMesh->GetEdgeVertex( EdgeID, 1 ) );
 		WireframeMesh->SetEdgeColor( EdgeID, GetEdgeColor( EditableMesh->GetEdgeAttribute( EdgeID, UEditableMeshAttribute::EdgeIsHard(), 0 ).X > 0.0f ) );
 	}
 
 	// Add all polygons and edge instances
 	// @todo mesheditor wireframe: holes not yet handled
-	for( TSparseArray<FMeshPolygon>::TConstIterator It( EditableMesh->Polygons ); It; ++It )
+	for( const FPolygonID PolygonID : EditableMesh->GetMeshDescription()->Polygons().GetElementIDs() )
 	{
-		const FPolygonID PolygonID( It.GetIndex() );
-
 		WireframeMesh->AddPolygon( PolygonID );
 		WireframeMesh->SetPolygonNormal( PolygonID, EditableMesh->ComputePolygonNormal( PolygonID ) );
 
@@ -87,7 +83,7 @@ void UMeshEditorStaticMeshAdapter::OnRebuildRenderMeshFinish( const UEditableMes
 {
 	for( const FPolygonID PolygonID : EditableMesh->PolygonsPendingNewTangentBasis )
 	{
-		const FMeshPolygon& Polygon = EditableMesh->Polygons[ PolygonID.GetValue() ];
+		const FMeshPolygon& Polygon = EditableMesh->GetMeshDescription()->GetPolygon( PolygonID );
 		WireframeMesh->SetPolygonNormal( PolygonID, Polygon.PolygonNormal );
 	}
 
@@ -214,12 +210,14 @@ void UMeshEditorStaticMeshAdapter::OnSetVertexInstanceAttribute( const UEditable
 
 void UMeshEditorStaticMeshAdapter::OnCreateEdges( const UEditableMesh* EditableMesh, const TArray<FEdgeID>& EdgeIDs )
 {
+	const FEdgeArray& Edges = EditableMesh->GetMeshDescription()->Edges();
+
 	for( const FEdgeID EdgeID : EdgeIDs )
 	{
-		const FMeshEdge& Edge = EditableMesh->Edges[ EdgeID.GetValue() ];
+		const FMeshEdge& Edge = Edges[ EdgeID ];
 
 		WireframeMesh->AddEdge( EdgeID );
-		WireframeMesh->SetEdgeVertices( EdgeID, Edge.VertexIDs[ 0 ], Edge.VertexIDs[ 1 ] );
+		WireframeMesh->SetEdgeVertices( EdgeID, EditableMesh->GetEdgeVertex( EdgeID, 0 ), EditableMesh->GetEdgeVertex( EdgeID, 1 ) );
 		WireframeMesh->SetEdgeColor( EdgeID, GetEdgeColor( EditableMesh->GetEdgeAttribute( EdgeID, UEditableMeshAttribute::EdgeIsHard(), 0 ).X > 0.0f ) );
 
 		for( const FPolygonID PolygonID : Edge.ConnectedPolygons )
@@ -232,9 +230,11 @@ void UMeshEditorStaticMeshAdapter::OnCreateEdges( const UEditableMesh* EditableM
 
 void UMeshEditorStaticMeshAdapter::OnDeleteEdges( const UEditableMesh* EditableMesh, const TArray<FEdgeID>& EdgeIDs )
 {
+	const FEdgeArray& Edges = EditableMesh->GetMeshDescription()->Edges();
+
 	for( const FEdgeID EdgeID : EdgeIDs )
 	{
-		const FMeshEdge& Edge = EditableMesh->Edges[ EdgeID.GetValue() ];
+		const FMeshEdge& Edge = Edges[ EdgeID ];
 
 		for( const FPolygonID PolygonID : Edge.ConnectedPolygons )
 		{
diff --git a/Engine/Plugins/Editor/MeshEditor/Source/MeshEditor/MeshEditorStaticMeshAdapter.h b/Engine/Plugins/Editor/MeshEditor/Source/MeshEditor/MeshEditorStaticMeshAdapter.h
--- a/Engine/Plugins/Editor/MeshEditor/Source/MeshEditor/MeshEditorStaticMeshAdapter.h
+++ b/Engine/Plugins/Editor/MeshEditor/Source/MeshEditor/MeshEditorStaticMeshAdapter.h
@@ -53,6 +53,7 @@ public:
 	virtual void OnChangePolygonVertexInstances( const UEditableMesh* EditableMesh, const TArray<FPolygonID>& PolygonIDs ) override;
 	virtual void OnCreatePolygonGroups( const UEditableMesh* EditableMesh, const TArray<FPolygonGroupID>& PolygonGroupIDs ) override;
 	virtual void OnDeletePolygonGroups( const UEditableMesh* EditableMesh, const TArray<FPolygonGroupID>& PolygonGroupIDs ) override;
+	virtual void OnAssignPolygonsToPolygonGroups( const UEditableMesh* EditableMesh, const TArray<FPolygonGroupForPolygon>& PolygonGroupForPolygons ) override {}
 	virtual void OnRetriangulatePolygons( const UEditableMesh* EditableMesh, const TArray<FPolygonID>& PolygonIDs ) override;
 
 
diff --git a/Engine/Plugins/Editor/MeshEditor/Source/MeshEditor/MeshEditorSubdividedStaticMeshAdapter.h b/Engine/Plugins/Editor/MeshEditor/Source/MeshEditor/MeshEditorSubdividedStaticMeshAdapter.h
--- a/Engine/Plugins/Editor/MeshEditor/Source/MeshEditor/MeshEditorSubdividedStaticMeshAdapter.h
+++ b/Engine/Plugins/Editor/MeshEditor/Source/MeshEditor/MeshEditorSubdividedStaticMeshAdapter.h
@@ -53,6 +53,7 @@ public:
 	virtual void OnChangePolygonVertexInstances( const UEditableMesh* EditableMesh, const TArray<FPolygonID>& PolygonIDs ) override;
 	virtual void OnCreatePolygonGroups( const UEditableMesh* EditableMesh, const TArray<FPolygonGroupID>& PolygonGroupIDs ) override;
 	virtual void OnDeletePolygonGroups( const UEditableMesh* EditableMesh, const TArray<FPolygonGroupID>& PolygonGroupIDs ) override;
+	virtual void OnAssignPolygonsToPolygonGroups( const UEditableMesh* EditableMesh, const TArray<FPolygonGroupForPolygon>& PolygonGroupForPolygons ) override {}
 	virtual void OnRetriangulatePolygons( const UEditableMesh* EditableMesh, const TArray<FPolygonID>& PolygonIDs ) override;
 
 
diff --git a/Engine/Plugins/Editor/MeshEditor/Source/PolygonModeling/PolygonModeling.Build.cs b/Engine/Plugins/Editor/MeshEditor/Source/PolygonModeling/PolygonModeling.Build.cs
--- a/Engine/Plugins/Editor/MeshEditor/Source/PolygonModeling/PolygonModeling.Build.cs
+++ b/Engine/Plugins/Editor/MeshEditor/Source/PolygonModeling/PolygonModeling.Build.cs
@@ -16,6 +16,7 @@ namespace UnrealBuildTool.Rules
 					"UnrealEd",
 					"EditableMesh",
 					"MeshEditor",
+                    "MeshDescription",
 					"SlateCore",
 					"ViewportInteraction"
 				}
diff --git a/Engine/Plugins/Runtime/EditableMesh/Source/EditableMesh/EditableMesh.Build.cs b/Engine/Plugins/Runtime/EditableMesh/Source/EditableMesh/EditableMesh.Build.cs
--- a/Engine/Plugins/Runtime/EditableMesh/Source/EditableMesh/EditableMesh.Build.cs
+++ b/Engine/Plugins/Runtime/EditableMesh/Source/EditableMesh/EditableMesh.Build.cs
@@ -11,11 +11,12 @@ namespace UnrealBuildTool.Rules
 				    "Core",
 				    "CoreUObject",
 				    "Engine",
-					"RenderCore",	// @todo mesheditor: For FlushRenderingCommands()
+                    "MeshDescription",
+					"RenderCore"	// @todo mesheditor: For FlushRenderingCommands()
 				}
 			);
 
 			AddEngineThirdPartyPrivateStaticDependencies(Target, new string[] { "MikkTSpace", "OpenSubdiv" });
-		}
-	}
+        }
+    }
 }
\ No newline at end of file
