Fix the mesh editor hit raycast

diff --git a/Engine/Plugins/Editor/MeshEditor/Source/MeshEditor/MeshEditorMode.cpp b/Engine/Plugins/Editor/MeshEditor/Source/MeshEditor/MeshEditorMode.cpp
--- a/Engine/Plugins/Editor/MeshEditor/Source/MeshEditor/MeshEditorMode.cpp
+++ b/Engine/Plugins/Editor/MeshEditor/Source/MeshEditor/MeshEditorMode.cpp
@@ -44,6 +44,8 @@
 #include "Framework/Application/SlateApplication.h"
 #include "Materials/Material.h"
 #include "Components/PrimitiveComponent.h"
+#include "ILevelViewport.h"
+#include "SLevelViewport.h"
 
 #define LOCTEXT_NAMESPACE "MeshEditorMode"
 
@@ -895,6 +897,23 @@ void FMeshEditorMode::Enter()
 		Toolkit->Init( Owner->GetToolkitHost() );
 	}
 
+	// Set the current viewport.
+	{
+		const TSharedRef< ILevelEditor >& LevelEditor = FModuleManager::GetModuleChecked<FLevelEditorModule>("LevelEditor").GetFirstLevelEditor().ToSharedRef();
+
+		// Do we have an active perspective viewport that is valid for VR?  If so, go ahead and use that.
+		TSharedPtr<FEditorViewportClient> ViewportClient;
+		{
+			TSharedPtr<ILevelViewport> ActiveLevelViewport = LevelEditor->GetActiveViewportInterface();
+			if (ActiveLevelViewport.IsValid())
+			{
+				ViewportClient = StaticCastSharedRef<SLevelViewport>(ActiveLevelViewport->AsWidget())->GetViewportClient();
+			}
+		}
+
+		ViewportWorldInteraction->SetDefaultOptionalViewportClient(ViewportClient);
+	}
+
 	UpdateSelectedEditableMeshes();
 }
 
