Total revamp of mesh element attribute model.

Attributes now have a number of possible types (FVector, FVector4, FVector2D, float, int, bool, FName, UObject*) and are exposed as individual flat arrays, indexed by element ID. For example, vertex positions are essentially exposed as an array of FVector which can be directly accessed and modified. This has a number of advantages:
- It is completely extensible: new attributes can be created (even by a third party) and added to a mesh description without requiring a serialization version bump, or any change to the parent structures.
- This is more efficient in batch operations which deal with a number of mesh elements in one go.
- These attribute buffers can potentially be passed directly to third-party libraries without requiring any kind of transformation.
- The distinct types allow for a better representation of the attribute being specified, without invalid values being possible (cf representing a bool value in an FVector4).

Attributes also have default values, and a flags field which confers use-specific properties to them. Editable Mesh currently uses this to determine whether an attribute's value can be automatically initialized by lerping the values of its neighbours, as well as for identifying auto-generated attributes such as tangents/normals. This is desirable as it means that even unknown / third-party attributes can potentially be handled transparently by [REDACTED_PERSON], without requiring the code to be extended.

Certain higher-level operations in EditableMesh have been optimized to make full use of vertex instances where possible.
The welding/splitting of identical vertex instances has been removed from here, as the aim is to unify this with mesh utility code elsewhere.
Various bug fixes.

diff --git a/Engine/Plugins/Editor/MeshEditor/Source/MeshEditor/MeshEditorMode.cpp b/Engine/Plugins/Editor/MeshEditor/Source/MeshEditor/MeshEditorMode.cpp
--- a/Engine/Plugins/Editor/MeshEditor/Source/MeshEditor/MeshEditorMode.cpp
+++ b/Engine/Plugins/Editor/MeshEditor/Source/MeshEditor/MeshEditorMode.cpp
@@ -5,6 +5,7 @@
 #include "MeshEditorStyle.h"
 #include "MeshEditorModeToolkit.h"
 #include "EditableMesh.h"
+#include "MeshAttributes.h"
 #include "EditableMeshFactory.h"
 #include "MeshElement.h"
 #include "DynamicMeshBuilder.h"
@@ -1301,20 +1302,24 @@ void FMeshEditorMode::UpdateDebugNormals()
 	{
 		const UPrimitiveComponent* Component = ComponentAndEditableMesh.Component;
 		const UEditableMesh* EditableMesh = ComponentAndEditableMesh.EditableMesh;
+		const UMeshDescription* MeshDescription = EditableMesh->GetMeshDescription();
+
+		const TVertexAttributeArray<FVector>& VertexPositions = MeshDescription->VertexAttributes().GetAttributes<FVector>( MeshAttribute::Vertex::Position );
+		const TVertexInstanceAttributeArray<FVector>& VertexNormals = MeshDescription->VertexInstanceAttributes().GetAttributes<FVector>( MeshAttribute::VertexInstance::Normal );
+		const TVertexInstanceAttributeArray<FVector>& VertexTangents = MeshDescription->VertexInstanceAttributes().GetAttributes<FVector>( MeshAttribute::VertexInstance::Tangent );
 
 		const FMatrix ComponentToWorldMatrix = Component->GetRenderMatrix();
 
-		for( const FPolygonID PolygonID : EditableMesh->GetMeshDescription()->Polygons().GetElementIDs() )
+		for( const FPolygonID PolygonID : MeshDescription->Polygons().GetElementIDs() )
 		{
 			// @todo mesheditor: total debug feature for now. Need a way of making this look nice.
-			const int32 PerimeterVertexCount = EditableMesh->GetPolygonPerimeterVertexCount( PolygonID );
 			const float Length = 10.0f; // @todo mesheditor: determine length of debug line from distance from the mesh origin to the camera?
 
-			for( int32 PerimeterVertexIndex = 0; PerimeterVertexIndex < PerimeterVertexCount; ++PerimeterVertexIndex )
+			for( const FVertexInstanceID VertexInstanceID : MeshDescription->GetPolygonPerimeterVertexInstances( PolygonID ) )
 			{
-				const FVector Position( EditableMesh->GetPolygonPerimeterVertexAttribute( PolygonID, PerimeterVertexIndex, UEditableMeshAttribute::VertexPosition(), 0 ) );
-				const FVector Normal( EditableMesh->GetPolygonPerimeterVertexAttribute( PolygonID, PerimeterVertexIndex, UEditableMeshAttribute::VertexNormal(), 0 ) );
-				const FVector Tangent( EditableMesh->GetPolygonPerimeterVertexAttribute( PolygonID, PerimeterVertexIndex, UEditableMeshAttribute::VertexTangent(), 0 ) );
+				const FVector Position = VertexPositions[ MeshDescription->GetVertexInstanceVertex( VertexInstanceID ) ];
+				const FVector Normal = VertexNormals[ VertexInstanceID ];
+				const FVector Tangent = VertexTangents[ VertexInstanceID ];
 
 				const FVector Start( ComponentToWorldMatrix.TransformPosition( Position ) );
 				const FVector NormalEnd( ComponentToWorldMatrix.TransformPosition( Position + Normal * Length ) );
@@ -1881,6 +1886,9 @@ void FMeshEditorMode::FrameSelectedElements( FEditorViewportClient* ViewportClie
 				UEditableMesh* EditableMesh = SelectedMeshAndVertices.Key;
 				const TArray<FMeshElement>& VertexElements = SelectedMeshAndVertices.Value;
 
+				const UMeshDescription* MeshDescription = EditableMesh->GetMeshDescription();
+				const TVertexAttributeArray<FVector>& VertexPositions = MeshDescription->VertexAttributes().GetAttributes<FVector>( MeshAttribute::Vertex::Position );
+
 				for( const auto& VertexElement : VertexElements )
 				{
 					const FVertexID VertexID( VertexElement.ElementAddress.ElementID );
@@ -1888,7 +1896,7 @@ void FMeshEditorMode::FrameSelectedElements( FEditorViewportClient* ViewportClie
 					const UPrimitiveComponent* Component = VertexElement.Component.Get();
 					if( Component )
 					{
-						const FVector VertexPosition = EditableMesh->GetVertexAttribute( VertexID, UEditableMeshAttribute::VertexPosition(), 0 );
+						const FVector VertexPosition = VertexPositions[ VertexID ];
 						BoundingBox += Component->GetComponentTransform().TransformPosition( VertexPosition );
 					}
 				}
@@ -1905,6 +1913,9 @@ void FMeshEditorMode::FrameSelectedElements( FEditorViewportClient* ViewportClie
 				UEditableMesh* EditableMesh = SelectedMeshAndEdges.Key;
 				const TArray<FMeshElement>& EdgeElements = SelectedMeshAndEdges.Value;
 
+				const UMeshDescription* MeshDescription = EditableMesh->GetMeshDescription();
+				const TVertexAttributeArray<FVector>& VertexPositions = MeshDescription->VertexAttributes().GetAttributes<FVector>( MeshAttribute::Vertex::Position );
+
 				for( const auto& EdgeElement : EdgeElements )
 				{
 					const FEdgeID EdgeID( EdgeElement.ElementAddress.ElementID );
@@ -1913,11 +1924,11 @@ void FMeshEditorMode::FrameSelectedElements( FEditorViewportClient* ViewportClie
 					if( Component )
 					{
 						const FVertexID VertexID0 = EditableMesh->GetEdgeVertex( EdgeID, 0 );
-						const FVector VertexPosition0 = EditableMesh->GetVertexAttribute( VertexID0, UEditableMeshAttribute::VertexPosition(), 0 );
+						const FVector VertexPosition0 = VertexPositions[ VertexID0 ];
 						BoundingBox += Component->GetComponentTransform().TransformPosition( VertexPosition0 );
 
 						const FVertexID VertexID1 = EditableMesh->GetEdgeVertex( EdgeID, 1 );
-						const FVector VertexPosition1 = EditableMesh->GetVertexAttribute( VertexID1, UEditableMeshAttribute::VertexPosition(), 0 );
+						const FVector VertexPosition1 = VertexPositions[ VertexID1 ];
 						BoundingBox += Component->GetComponentTransform().TransformPosition( VertexPosition1 );
 					}
 				}
@@ -1934,6 +1945,9 @@ void FMeshEditorMode::FrameSelectedElements( FEditorViewportClient* ViewportClie
 				UEditableMesh* EditableMesh = SelectedMeshAndPolygons.Key;
 				const TArray<FMeshElement>& PolygonElements = SelectedMeshAndPolygons.Value;
 
+				const UMeshDescription* MeshDescription = EditableMesh->GetMeshDescription();
+				const TVertexAttributeArray<FVector>& VertexPositions = MeshDescription->VertexAttributes().GetAttributes<FVector>( MeshAttribute::Vertex::Position );
+
 				for( const auto& PolygonElement : PolygonElements )
 				{
 					const UPrimitiveComponent* Component = PolygonElement.Component.Get();
@@ -1941,10 +1955,9 @@ void FMeshEditorMode::FrameSelectedElements( FEditorViewportClient* ViewportClie
 					{
 						const FPolygonID PolygonID( PolygonElement.ElementAddress.ElementID );
 
-						const int32 PolygonVertexCount = EditableMesh->GetPolygonPerimeterVertexCount( PolygonID );
-						for( int32 PolygonVertexIndex = 0; PolygonVertexIndex < PolygonVertexCount; ++PolygonVertexIndex )
+						for( const FVertexInstanceID VertexInstanceID : MeshDescription->GetPolygonPerimeterVertexInstances( PolygonID ) )
 						{
-							const FVector VertexPosition = EditableMesh->GetPolygonPerimeterVertexAttribute( PolygonID, PolygonVertexIndex, UEditableMeshAttribute::VertexPosition(), 0 );
+							const FVector VertexPosition = VertexPositions[ MeshDescription->GetVertexInstanceVertex( VertexInstanceID ) ];
 							BoundingBox += Component->GetComponentTransform().TransformPosition( VertexPosition );
 						}
 					}
@@ -2270,6 +2283,7 @@ bool FMeshEditorMode::AssignMaterialToSelectedPolygons( UMaterialInterface* Sele
 		for( const auto& MeshAndPolygons : MeshesAndPolygons )
 		{
 			UEditableMesh* EditableMesh = MeshAndPolygons.Key;
+			const UMeshDescription* MeshDescription = EditableMesh->GetMeshDescription();
 
 			UPrimitiveComponent* Component = nullptr;
 			for( const FMeshElement& PolygonElement : MeshAndPolygons.Value )
@@ -2286,9 +2300,12 @@ bool FMeshEditorMode::AssignMaterialToSelectedPolygons( UMaterialInterface* Sele
 				// Eventually we will need to be able to specify PolygonGroup properties in the editor, and ask the user for
 				// further details if there is more than one polygon group which matches the material.
 				FPolygonGroupID PolygonGroupToAssign = FPolygonGroupID::Invalid;
-				for( const FPolygonGroupID PolygonGroupID : EditableMesh->GetMeshDescription()->PolygonGroups().GetElementIDs() )
+
+				const TPolygonGroupAttributeArray<UObject*>& PolygonGroupMaterialAssets = MeshDescription->PolygonGroupAttributes().GetAttributes<UObject*>( MeshAttribute::PolygonGroup::MaterialAsset );
+
+				for( const FPolygonGroupID PolygonGroupID : MeshDescription->PolygonGroups().GetElementIDs() )
 				{
-					UMaterialInterface* Material = EditableMesh->GetPolygonGroupMaterial( PolygonGroupID );
+					UMaterialInterface* Material = Cast<UMaterialInterface>( PolygonGroupMaterialAssets[ PolygonGroupID ] );
 					if( Material == SelectedMaterial )
 					{
 						// We only expect to find one polygon group containing this material at the moment.
@@ -2304,11 +2321,12 @@ bool FMeshEditorMode::AssignMaterialToSelectedPolygons( UMaterialInterface* Sele
 				{
 					// Helper function which returns a unique FName for the material slot name, based on the material's asset name,
 					// and adding a unique suffix if there are other polygon groups with the same material slot name.
-					auto MakeUniqueSlotName = []( UEditableMesh* EditableMesh, FName Name ) -> FName
+					auto MakeUniqueSlotName = []( const UMeshDescription* MeshDescription, FName Name ) -> FName
 					{
-						for( const FPolygonGroupID PolygonGroupID : EditableMesh->GetMeshDescription()->PolygonGroups().GetElementIDs() )
+						const TPolygonGroupAttributeArray<FName>& MaterialSlotNames = MeshDescription->PolygonGroupAttributes().GetAttributes<FName>( MeshAttribute::PolygonGroup::MaterialSlotName );
+						for( const FPolygonGroupID PolygonGroupID : MeshDescription->PolygonGroups().GetElementIDs() )
 						{
-							const FName ExistingName = EditableMesh->GetPolygonGroupMaterialSlotName( PolygonGroupID );
+							const FName ExistingName = MaterialSlotNames[ PolygonGroupID ];
 							if( ExistingName.GetComparisonIndex() == Name.GetComparisonIndex() )
 							{
 								Name = FName( Name, FMath::Max( Name.GetNumber(), ExistingName.GetNumber() + 1 ) );
@@ -2317,15 +2335,17 @@ bool FMeshEditorMode::AssignMaterialToSelectedPolygons( UMaterialInterface* Sele
 						return Name;
 					};
 
+					const FName UniqueSlotName = MakeUniqueSlotName( MeshDescription, SelectedMaterial->GetFName() );
+
 					static TArray<FPolygonGroupToCreate> PolygonGroupsToCreate;
 					PolygonGroupsToCreate.Reset( 1 );
 					PolygonGroupsToCreate.Emplace();
 					
 					FPolygonGroupToCreate& PolygonGroupToCreate = PolygonGroupsToCreate.Last();
-					PolygonGroupToCreate.MaterialAsset = SelectedMaterial->GetPathName();
-					PolygonGroupToCreate.MaterialSlotName = MakeUniqueSlotName( EditableMesh, SelectedMaterial->GetFName() );
-					PolygonGroupToCreate.bEnableCollision = true;
-					PolygonGroupToCreate.bCastShadow = true;
+					PolygonGroupToCreate.PolygonGroupAttributes.Attributes.Emplace( MeshAttribute::PolygonGroup::MaterialAsset, 0, FMeshElementAttributeValue( SelectedMaterial ) );
+					PolygonGroupToCreate.PolygonGroupAttributes.Attributes.Emplace( MeshAttribute::PolygonGroup::MaterialSlotName, 0, FMeshElementAttributeValue( UniqueSlotName ) );
+					PolygonGroupToCreate.PolygonGroupAttributes.Attributes.Emplace( MeshAttribute::PolygonGroup::CastShadow, 0, FMeshElementAttributeValue( true ) );
+					PolygonGroupToCreate.PolygonGroupAttributes.Attributes.Emplace( MeshAttribute::PolygonGroup::EnableCollision, 0, FMeshElementAttributeValue( true ) );
 
 					static TArray<FPolygonGroupID> NewPolygonGroupIDs;
 					EditableMesh->CreatePolygonGroups( PolygonGroupsToCreate, NewPolygonGroupIDs );
@@ -2426,6 +2446,9 @@ void FMeshEditorMode::AddMeshElementToOverlay( UOverlayComponent* OverlayCompone
 		{
 			if( IsElementIDValid( MeshElement, EditableMesh ) )
 			{
+				const UMeshDescription* MeshDescription = EditableMesh->GetMeshDescription();
+				const TVertexAttributeArray<FVector>& VertexPositions = MeshDescription->VertexAttributes().GetAttributes<FVector>( MeshAttribute::Vertex::Position );
+
 				UPrimitiveComponent* Component = MeshElement.Component.Get();
 				check( Component != nullptr );
 
@@ -2436,7 +2459,7 @@ void FMeshEditorMode::AddMeshElementToOverlay( UOverlayComponent* OverlayCompone
 					case EEditableMeshElementType::Vertex:
 					{
 						const FVertexID VertexID( MeshElement.ElementAddress.ElementID );
-						const FVector Position( ComponentToWorldMatrix.TransformPosition( EditableMesh->GetVertexAttribute( VertexID, UEditableMeshAttribute::VertexPosition(), 0 ) ) );
+						const FVector Position( ComponentToWorldMatrix.TransformPosition( VertexPositions[ VertexID ] ) );
 
 						OverlayComponent->AddPoint( FOverlayPoint( Position, Color, Size ) );
 						break;
@@ -2447,8 +2470,8 @@ void FMeshEditorMode::AddMeshElementToOverlay( UOverlayComponent* OverlayCompone
 						const FEdgeID EdgeID( MeshElement.ElementAddress.ElementID );
 						const FVertexID StartVertexID( EditableMesh->GetEdgeVertex( EdgeID, 0 ) );
 						const FVertexID EndVertexID( EditableMesh->GetEdgeVertex( EdgeID, 1 ) );
-						const FVector StartPosition( ComponentToWorldMatrix.TransformPosition( EditableMesh->GetVertexAttribute( StartVertexID, UEditableMeshAttribute::VertexPosition(), 0 ) ) );
-						const FVector EndPosition( ComponentToWorldMatrix.TransformPosition( EditableMesh->GetVertexAttribute( EndVertexID, UEditableMeshAttribute::VertexPosition(), 0 ) ) );
+						const FVector StartPosition( ComponentToWorldMatrix.TransformPosition( VertexPositions[ StartVertexID ] ) );
+						const FVector EndPosition( ComponentToWorldMatrix.TransformPosition( VertexPositions[ EndVertexID ] ) );
 
 						OverlayComponent->AddLine( FOverlayLine( StartPosition, EndPosition, Color, Size ) );
 						break;
@@ -2461,19 +2484,19 @@ void FMeshEditorMode::AddMeshElementToOverlay( UOverlayComponent* OverlayCompone
 
 						for( int32 PolygonTriangle = 0; PolygonTriangle < PolygonTriangleCount; PolygonTriangle++ )
 						{
-							FVector VertexPositions[ 3 ];
+							FVector TriangleVertexPositions[ 3 ];
 							for( int32 TriangleVertex = 0; TriangleVertex < 3; TriangleVertex++ )
 							{
 								const FVertexInstanceID VertexInstanceID = EditableMesh->GetPolygonTriangulatedTriangle( PolygonID, PolygonTriangle ).GetVertexInstanceID( TriangleVertex );
 								const FVertexID VertexID = EditableMesh->GetVertexInstanceVertex( VertexInstanceID );
-								VertexPositions[ TriangleVertex ] = ComponentToWorldMatrix.TransformPosition( EditableMesh->GetVertexAttribute( VertexID, UEditableMeshAttribute::VertexPosition(), 0 ) );
+								TriangleVertexPositions[ TriangleVertex ] = ComponentToWorldMatrix.TransformPosition( VertexPositions[ VertexID ] );
 							}
 
 							OverlayComponent->AddTriangle( FOverlayTriangle( 
 								HoveredFaceMaterial,
-								FOverlayTriangleVertex( VertexPositions[ 0 ], FVector2D( 0, 0 ), FVector::UpVector, Color ),
-								FOverlayTriangleVertex( VertexPositions[ 1 ], FVector2D( 0, 1 ), FVector::UpVector, Color ),
-								FOverlayTriangleVertex( VertexPositions[ 2 ], FVector2D( 1, 1 ), FVector::UpVector, Color )
+								FOverlayTriangleVertex( TriangleVertexPositions[ 0 ], FVector2D( 0, 0 ), FVector::UpVector, Color ),
+								FOverlayTriangleVertex( TriangleVertexPositions[ 1 ], FVector2D( 0, 1 ), FVector::UpVector, Color ),
+								FOverlayTriangleVertex( TriangleVertexPositions[ 2 ], FVector2D( 1, 1 ), FVector::UpVector, Color )
 							) );
 						}
 						break;
@@ -3127,9 +3150,9 @@ void FMeshEditorMode::UpdateActiveAction( const bool bIsActionFinishing )
 					FVertexToCreate& VertexToCreate = VerticesToCreate.Last();
 
 					VertexToCreate.VertexAttributes.Attributes.Emplace(
-						UEditableMeshAttribute::VertexPosition(),
+						MeshAttribute::Vertex::Position,
 						0,
-						FVector4( Component->GetComponentTransform().InverseTransformPosition( DrawnPoint.Get<1>() ) )
+						FMeshElementAttributeValue( Component->GetComponentTransform().InverseTransformPosition( DrawnPoint.Get<1>() ) )
 					);
 				}
 
@@ -3163,7 +3186,7 @@ void FMeshEditorMode::UpdateActiveAction( const bool bIsActionFinishing )
 					EdgeToCreate.VertexID0 = NewVertexIDs[ 0 ];
 					EdgeToCreate.VertexID1 = NewVertexIDs[ 1 ];
 
-					EdgeToCreate.EdgeAttributes.Attributes.Emplace( UEditableMeshAttribute::EdgeIsHard(), 0, FVector4( 1.0f ) );
+					EdgeToCreate.EdgeAttributes.Attributes.Emplace( MeshAttribute::Edge::IsHard, 0, FMeshElementAttributeValue( true ) );
 
 					EditableMesh->CreateEdges( EdgesToCreate, NewEdgeIDs );
 				}
@@ -3222,9 +3245,9 @@ void FMeshEditorMode::UpdateActiveAction( const bool bIsActionFinishing )
 						FMeshElementAttributeList& AttributeList = VertexAttrs.PerimeterVertexAttributeLists.Last();
 
 						AttributeList.Attributes.Emplace(
-							UEditableMeshAttribute::VertexNormal(),
+							MeshAttribute::VertexInstance::Normal,
 							0,
-							FVector4( PolygonNormal )
+							FMeshElementAttributeValue( PolygonNormal )
 						);
 					}
 				}
@@ -3301,6 +3324,8 @@ void FMeshEditorMode::UpdateActiveAction( const bool bIsActionFinishing )
 			VertexIDsAlreadyMoved.Reset();
 
 			UEditableMesh* EditableMesh = MeshAndTransformables.Key;
+			const UMeshDescription* MeshDescription = EditableMesh->GetMeshDescription();
+			const TVertexAttributeArray<FVector>& VertexPositions = MeshDescription->VertexAttributes().GetAttributes<FVector>( MeshAttribute::Vertex::Position );
 
 			const TArray< const FMeshElementViewportTransformable* >& TransformablesForMesh = MeshAndTransformables.Value;
 
@@ -3351,7 +3376,7 @@ void FMeshEditorMode::UpdateActiveAction( const bool bIsActionFinishing )
 					{
 						if( !VertexIDsAlreadyMoved.Contains( EdgeVertexID ) )
 						{
-							const FVector OriginalComponentSpaceVertexPosition = EditableMesh->GetVertexAttribute( EdgeVertexID, UEditableMeshAttribute::VertexPosition(), 0 );
+							const FVector OriginalComponentSpaceVertexPosition = VertexPositions[ EdgeVertexID ];
 							const FVector NewComponentSpaceVertexPosition = ComponentDeltaFromStartTransform.TransformPosition( OriginalComponentSpaceVertexPosition );
 
 							FVertexToMove& VertexToMove = *new( VerticesToMove ) FVertexToMove();
@@ -3373,7 +3398,7 @@ void FMeshEditorMode::UpdateActiveAction( const bool bIsActionFinishing )
 					{
 						if( !VertexIDsAlreadyMoved.Contains( PolygonPerimeterVertexID ) )
 						{
-							const FVector OriginalComponentSpaceVertexPosition = EditableMesh->GetVertexAttribute( PolygonPerimeterVertexID, UEditableMeshAttribute::VertexPosition(), 0 );
+							const FVector OriginalComponentSpaceVertexPosition = VertexPositions[ PolygonPerimeterVertexID ];
 							const FVector NewComponentSpaceVertexPosition = ComponentDeltaFromStartTransform.TransformPosition( OriginalComponentSpaceVertexPosition );
 
 							FVertexToMove& VertexToMove = *new( VerticesToMove ) FVertexToMove();
@@ -3452,6 +3477,9 @@ bool FMeshEditorMode::FindEdgeSplitUnderInteractor(	UViewportInteractor* Viewpor
 	OutClosestEdgeID = FEdgeID::Invalid;
 	bool bFoundSplit = false;
 
+	const UMeshDescription* MeshDescription = EditableMesh->GetMeshDescription();
+	const TVertexAttributeArray<FVector>& VertexPositions = MeshDescription->VertexAttributes().GetAttributes<FVector>( MeshAttribute::Vertex::Position );
+
 	// Figure out where to split based on where the interactor is aiming.  We'll look at all of the
 	// selected edges, and choose a split offset based on the closest point along one of those edges
 	// to the interactor.  All selected edges will then be split by the same proportion.
@@ -3472,7 +3500,7 @@ bool FMeshEditorMode::FindEdgeSplitUnderInteractor(	UViewportInteractor* Viewpor
 			for( int32 EdgeVertexNumber = 0; EdgeVertexNumber < 2; ++EdgeVertexNumber )
 			{
 				WorldSpaceEdgeVertexPositions[ EdgeVertexNumber ] =
-					ComponentToWorld.TransformPosition( EditableMesh->GetVertexAttribute( EdgeVertexIDs[ EdgeVertexNumber ], UEditableMeshAttribute::VertexPosition(), 0 ) );
+					ComponentToWorld.TransformPosition( VertexPositions[ EdgeVertexIDs[ EdgeVertexNumber ] ] );
 			}
 
 
@@ -3583,6 +3611,9 @@ FEditableMeshElementAddress FMeshEditorMode::QueryElement( const UEditableMesh&
 	float ClosestDistanceToRay = TNumericLimits<float>::Max();
 	FVector CurrentRayEnd = RayEnd;
 
+	const UMeshDescription* MeshDescription = EditableMesh.GetMeshDescription();
+	const TVertexAttributeArray<FVector>& VertexPositions = MeshDescription->VertexAttributes().GetAttributes<FVector>( MeshAttribute::Vertex::Position );
+
 	// Check polygons first; this is so we always impose a closest hit location at the poly before checking other elements, so anything behind is occluded
 	for( const FPolygonID PolygonID : FrontFacingPolygons )
 	{
@@ -3598,7 +3629,7 @@ FEditableMeshElementAddress FMeshEditorMode::QueryElement( const UEditableMesh&
 			{
 				const FVertexInstanceID VertexInstanceID = EditableMesh.GetPolygonTriangulatedTriangle( PolygonID, PolygonTriangleNumber ).GetVertexInstanceID( TriangleVertexNumber );
 				const FVertexID VertexID = EditableMesh.GetVertexInstanceVertex( VertexInstanceID );
-				TriangleVertexPositions[ TriangleVertexNumber ] = EditableMesh.GetVertexAttribute( VertexID, UEditableMeshAttribute::VertexPosition(), 0 );
+				TriangleVertexPositions[ TriangleVertexNumber ] = VertexPositions[ VertexID ];
 			}
 
 			const bool bAlreadyHitTriangle = ( HitElementAddress.ElementType == EEditableMeshElementType::Polygon );
@@ -3620,8 +3651,8 @@ FEditableMeshElementAddress FMeshEditorMode::QueryElement( const UEditableMesh&
 		for( const FEdgeID EdgeID : FrontFacingEdges )
 		{
 			FVector EdgeVertexPositions[ 2 ];
-			EdgeVertexPositions[ 0 ] = EditableMesh.GetVertexAttribute( EditableMesh.GetEdgeVertex( EdgeID, 0 ), UEditableMeshAttribute::VertexPosition(), 0 );
-			EdgeVertexPositions[ 1 ] = EditableMesh.GetVertexAttribute( EditableMesh.GetEdgeVertex( EdgeID, 1 ), UEditableMeshAttribute::VertexPosition(), 0 );
+			EdgeVertexPositions[ 0 ] = VertexPositions[ EditableMesh.GetEdgeVertex( EdgeID, 0 ) ];
+			EdgeVertexPositions[ 1 ] = VertexPositions[ EditableMesh.GetEdgeVertex( EdgeID, 1 ) ];
 
 			const bool bAlreadyHitEdge = ( HitElementAddress.ElementType == EEditableMeshElementType::Edge );
 			const bool bHit = CheckEdge( InteractorShape, Sphere, SphereFuzzyDistance, RayStart, CurrentRayEnd, RayFuzzyDistance, EdgeVertexPositions, CameraLocation, bIsPerspectiveView, FuzzyDistanceScaleFactor, ClosestInteractorShape, ClosestDistanceToRay, ClosestDistanceOnRay, ClosestHitLocation, bAlreadyHitEdge );
@@ -3640,7 +3671,7 @@ FEditableMeshElementAddress FMeshEditorMode::QueryElement( const UEditableMesh&
 	{
 		for( const FVertexID VertexID : FrontFacingVertices )
 		{
-			const FVector VertexPosition = EditableMesh.GetVertexAttribute( VertexID, UEditableMeshAttribute::VertexPosition(), 0 );
+			const FVector VertexPosition = VertexPositions[ VertexID ];
 			const bool bAlreadyHitVertex = ( HitElementAddress.ElementType == EEditableMeshElementType::Vertex );
 			const bool bHit = CheckVertex( InteractorShape, Sphere, SphereFuzzyDistance, RayStart, CurrentRayEnd, RayFuzzyDistance, VertexPosition, CameraLocation, bIsPerspectiveView, FuzzyDistanceScaleFactor, ClosestInteractorShape, ClosestDistanceToRay, ClosestDistanceOnRay, ClosestHitLocation, bAlreadyHitVertex );
 			if( bHit )
@@ -4384,6 +4415,9 @@ bool FMeshEditorMode::FrustumSelect( const FConvexVolume& InFrustum, FEditorView
 		UEditableMesh* EditableMesh = CandidateMesh.Get<1>();
 		FTransform ComponentTransform = Component->GetComponentTransform();
 
+		const UMeshDescription* MeshDescription = EditableMesh->GetMeshDescription();
+		const TVertexAttributeArray<FVector>& VertexPositions = MeshDescription->VertexAttributes().GetAttributes<FVector>( MeshAttribute::Vertex::Position );
+
 		static TArray<FEdgeID> SelectedEdgeIDs;
 		SelectedEdgeIDs.Empty();
 
@@ -4399,8 +4433,8 @@ bool FMeshEditorMode::FrustumSelect( const FConvexVolume& InFrustum, FEditorView
 			const FVertexID VertexID1( EditableMesh->GetEdgeVertex( EdgeID, 0 ) );
 			const FVertexID VertexID2( EditableMesh->GetEdgeVertex( EdgeID, 1 ) );
 
-			const FVector VertexPosition1( ComponentTransform.TransformPosition( EditableMesh->GetVertexAttribute( VertexID1, UEditableMeshAttribute::VertexPosition(), 0 ) ) );
-			const FVector VertexPosition2( ComponentTransform.TransformPosition( EditableMesh->GetVertexAttribute( VertexID2, UEditableMeshAttribute::VertexPosition(), 0 ) ) );
+			const FVector VertexPosition1( ComponentTransform.TransformPosition( VertexPositions[ VertexID1 ] ) );
+			const FVector VertexPosition2( ComponentTransform.TransformPosition( VertexPositions[ VertexID2 ] ) );
 
 			if( InFrustum.IntersectLineSegment( VertexPosition1, VertexPosition2 ) )
 			{
@@ -4455,7 +4489,7 @@ bool FMeshEditorMode::FrustumSelect( const FConvexVolume& InFrustum, FEditorView
 			// Here we only want to catch vertices with no associated edges.
 			if( EditableMesh->GetVertexConnectedEdgeCount( VertexID ) == 0 )
 			{
-				const FVector VertexPosition( ComponentTransform.TransformPosition( EditableMesh->GetVertexAttribute( VertexID, UEditableMeshAttribute::VertexPosition(), 0 ) ) );
+				const FVector VertexPosition( ComponentTransform.TransformPosition( VertexPositions[ VertexID ] ) );
 				if( InFrustum.IntersectPoint( VertexPosition ) )
 				{
 					// As the vertex is orphaned, it cannot be front or back facing. So we add it regardless.
@@ -4611,13 +4645,16 @@ void FMeshEditorMode::RefreshTransformables( const bool bNewObjectsSelected )
 					const FTransform ComponentToWorld = Component->GetComponentToWorld();
 					const FMatrix ComponentToWorldMatrix = Component->GetRenderMatrix();
 
+					const UMeshDescription* MeshDescription = EditableMesh->GetMeshDescription();
+					const TVertexAttributeArray<FVector>& VertexPositions = MeshDescription->VertexAttributes().GetAttributes<FVector>( MeshAttribute::Vertex::Position );
+
 					FTransform ElementTransform = FTransform::Identity;
 					switch( MeshElement.ElementAddress.ElementType )
 					{
 						case EEditableMeshElementType::Vertex:
 							ElementTransform.SetLocation(
 								ComponentToWorldMatrix.TransformPosition(
-									EditableMesh->GetVertexAttribute( FVertexID( MeshElement.ElementAddress.ElementID ), UEditableMeshAttribute::VertexPosition(), 0 )
+									VertexPositions[ FVertexID( MeshElement.ElementAddress.ElementID ) ]
 								)
 							);
 							break;
@@ -4629,8 +4666,8 @@ void FMeshEditorMode::RefreshTransformables( const bool bNewObjectsSelected )
 
 							const FVector ComponentSpaceEdgeCenter =
 								FMath::Lerp(
-									EditableMesh->GetVertexAttribute( EdgeVertexID0, UEditableMeshAttribute::VertexPosition(), 0 ),
-									EditableMesh->GetVertexAttribute( EdgeVertexID1, UEditableMeshAttribute::VertexPosition(), 0 ),
+									VertexPositions[ EdgeVertexID0 ],
+									VertexPositions[ EdgeVertexID1 ],
 									0.5f );
 
 							ElementTransform.SetLocation( ComponentToWorldMatrix.TransformPosition( ComponentSpaceEdgeCenter ) );
@@ -4645,8 +4682,8 @@ void FMeshEditorMode::RefreshTransformables( const bool bNewObjectsSelected )
 							ElementTransform.SetLocation( ComponentToWorldMatrix.TransformPosition( ComponentSpacePolygonCenter ) );
 
 							const FVector WindingVector =
-								( ComponentToWorldMatrix.TransformPosition( EditableMesh->GetVertexAttribute( EditableMesh->GetPolygonPerimeterVertex( PolygonID, 1 ), UEditableMeshAttribute::VertexPosition(), 0 ) ) -
-									ComponentToWorldMatrix.TransformPosition( EditableMesh->GetVertexAttribute( EditableMesh->GetPolygonPerimeterVertex( PolygonID, 0 ), UEditableMeshAttribute::VertexPosition(), 0 ) ) ).GetSafeNormal();
+								( ComponentToWorldMatrix.TransformPosition( VertexPositions[ EditableMesh->GetPolygonPerimeterVertex( PolygonID, 1 ) ] ) -
+									ComponentToWorldMatrix.TransformPosition( VertexPositions[ EditableMesh->GetPolygonPerimeterVertex( PolygonID, 0 ) ] ) ).GetSafeNormal();
 
 							const FVector PolygonNormal = ComponentToWorld.TransformVectorNoScale( EditableMesh->ComputePolygonNormal( PolygonID ) ).GetSafeNormal();
 
diff --git a/Engine/Plugins/Editor/MeshEditor/Source/MeshEditor/MeshEditorStaticMeshAdapter.cpp b/Engine/Plugins/Editor/MeshEditor/Source/MeshEditor/MeshEditorStaticMeshAdapter.cpp
--- a/Engine/Plugins/Editor/MeshEditor/Source/MeshEditor/MeshEditorStaticMeshAdapter.cpp
+++ b/Engine/Plugins/Editor/MeshEditor/Source/MeshEditor/MeshEditorStaticMeshAdapter.cpp
@@ -2,6 +2,7 @@
 
 #include "MeshEditorStaticMeshAdapter.h"
 #include "EditableMesh.h"
+#include "MeshAttributes.h"
 #include "WireframeMeshComponent.h"
 
 
@@ -47,19 +48,24 @@ void UMeshEditorStaticMeshAdapter::OnRebuildRenderMesh( const UEditableMesh* Edi
 {
 	WireframeMesh->Reset();
 
+	const UMeshDescription* MeshDescription = EditableMesh->GetMeshDescription();
+	check( MeshDescription );
+
 	// Add all vertices
-	for( const FVertexID VertexID : EditableMesh->GetMeshDescription()->Vertices().GetElementIDs() )
+	const TVertexAttributeArray<FVector>& VertexPositions = MeshDescription->VertexAttributes().GetAttributes<FVector>( MeshAttribute::Vertex::Position );
+	for( const FVertexID VertexID : MeshDescription->Vertices().GetElementIDs() )
 	{
 		WireframeMesh->AddVertex( VertexID );
-		WireframeMesh->SetVertexPosition( VertexID, EditableMesh->GetVertexAttribute( VertexID, UEditableMeshAttribute::VertexPosition(), 0 ) );
+		WireframeMesh->SetVertexPosition( VertexID, VertexPositions[ VertexID ] );
 	}
 
 	// Add all edges
-	for( const FEdgeID EdgeID : EditableMesh->GetMeshDescription()->Edges().GetElementIDs() )
+	const TEdgeAttributeArray<bool>& EdgeHardnesses = MeshDescription->EdgeAttributes().GetAttributes<bool>( MeshAttribute::Edge::IsHard );
+	for( const FEdgeID EdgeID : MeshDescription->Edges().GetElementIDs() )
 	{
 		WireframeMesh->AddEdge( EdgeID );
 		WireframeMesh->SetEdgeVertices( EdgeID, EditableMesh->GetEdgeVertex( EdgeID, 0 ), EditableMesh->GetEdgeVertex( EdgeID, 1 ) );
-		WireframeMesh->SetEdgeColor( EdgeID, GetEdgeColor( EditableMesh->GetEdgeAttribute( EdgeID, UEditableMeshAttribute::EdgeIsHard(), 0 ).X > 0.0f ) );
+		WireframeMesh->SetEdgeColor( EdgeID, GetEdgeColor( EdgeHardnesses[ EdgeID ] ) );
 	}
 
 	// Add all polygons and edge instances
@@ -81,10 +87,13 @@ void UMeshEditorStaticMeshAdapter::OnRebuildRenderMesh( const UEditableMesh* Edi
 
 void UMeshEditorStaticMeshAdapter::OnRebuildRenderMeshFinish( const UEditableMesh* EditableMesh, const bool bRebuildBoundsAndCollision, const bool bIsPreviewRollback )
 {
+	const UMeshDescription* MeshDescription = EditableMesh->GetMeshDescription();
+	check( MeshDescription );
+
+	const TPolygonAttributeArray<FVector>& PolygonNormals = MeshDescription->PolygonAttributes().GetAttributes<FVector>( MeshAttribute::Polygon::Normal );
 	for( const FPolygonID PolygonID : EditableMesh->PolygonsPendingNewTangentBasis )
 	{
-		const FMeshPolygon& Polygon = EditableMesh->GetMeshDescription()->GetPolygon( PolygonID );
-		WireframeMesh->SetPolygonNormal( PolygonID, Polygon.PolygonNormal );
+		WireframeMesh->SetPolygonNormal( PolygonID, PolygonNormals[ PolygonID ] );
 	}
 
 	WireframeMesh->InitResources();
@@ -181,10 +190,14 @@ void UMeshEditorStaticMeshAdapter::OnCreateEmptyVertexRange( const UEditableMesh
 
 void UMeshEditorStaticMeshAdapter::OnCreateVertices( const UEditableMesh* EditableMesh, const TArray<FVertexID>& VertexIDs )
 {
+	const UMeshDescription* MeshDescription = EditableMesh->GetMeshDescription();
+	check( MeshDescription );
+
+	const TVertexAttributeArray<FVector>& VertexPositions = MeshDescription->VertexAttributes().GetAttributes<FVector>( MeshAttribute::Vertex::Position );
 	for( const FVertexID VertexID : VertexIDs )
 	{
 		WireframeMesh->AddVertex( VertexID );
-		WireframeMesh->SetVertexPosition( VertexID, EditableMesh->GetVertexAttribute( VertexID, UEditableMeshAttribute::VertexPosition(), 0 ) );
+		WireframeMesh->SetVertexPosition( VertexID, VertexPositions[ VertexID ] );
 	}
 }
 
@@ -194,33 +207,29 @@ void UMeshEditorStaticMeshAdapter::OnCreateVertexInstances( const UEditableMesh*
 }
 
 
-void UMeshEditorStaticMeshAdapter::OnSetVertexAttribute( const UEditableMesh* EditableMesh, const FVertexID VertexID, const FName AttributeName, const int32 AttributeIndex, const FVector4 AttributeValue )
+void UMeshEditorStaticMeshAdapter::OnSetVertexAttribute( const UEditableMesh* EditableMesh, const FVertexID VertexID, const FMeshElementAttributeData& Attribute )
 {
-	if( AttributeName == UEditableMeshAttribute::VertexPosition() )
+	if( Attribute.AttributeName == MeshAttribute::Vertex::Position )
 	{
-		WireframeMesh->SetVertexPosition( VertexID, AttributeValue );
+		WireframeMesh->SetVertexPosition( VertexID, Attribute.AttributeValue.GetValue<FVector>() );
 	}
 }
 
 
-void UMeshEditorStaticMeshAdapter::OnSetVertexInstanceAttribute( const UEditableMesh* EditableMesh, const FVertexInstanceID VertexInstanceID, const FName AttributeName, const int32 AttributeIndex, const FVector4 AttributeValue )
-{
-}
-
-
 void UMeshEditorStaticMeshAdapter::OnCreateEdges( const UEditableMesh* EditableMesh, const TArray<FEdgeID>& EdgeIDs )
 {
-	const FEdgeArray& Edges = EditableMesh->GetMeshDescription()->Edges();
+	const UMeshDescription* MeshDescription = EditableMesh->GetMeshDescription();
+	check( MeshDescription );
+
+	const TEdgeAttributeArray<bool>& EdgeHardnesses = MeshDescription->EdgeAttributes().GetAttributes<bool>( MeshAttribute::Edge::IsHard );
 
 	for( const FEdgeID EdgeID : EdgeIDs )
 	{
-		const FMeshEdge& Edge = Edges[ EdgeID ];
-
 		WireframeMesh->AddEdge( EdgeID );
 		WireframeMesh->SetEdgeVertices( EdgeID, EditableMesh->GetEdgeVertex( EdgeID, 0 ), EditableMesh->GetEdgeVertex( EdgeID, 1 ) );
-		WireframeMesh->SetEdgeColor( EdgeID, GetEdgeColor( EditableMesh->GetEdgeAttribute( EdgeID, UEditableMeshAttribute::EdgeIsHard(), 0 ).X > 0.0f ) );
+		WireframeMesh->SetEdgeColor( EdgeID, GetEdgeColor( EdgeHardnesses[ EdgeID ] ) );
 
-		for( const FPolygonID PolygonID : Edge.ConnectedPolygons )
+		for( const FPolygonID PolygonID : MeshDescription->GetEdgeConnectedPolygons( EdgeID ) )
 		{
 			WireframeMesh->AddEdgeInstance( EdgeID, PolygonID );
 		}
@@ -230,13 +239,12 @@ void UMeshEditorStaticMeshAdapter::OnCreateEdges( const UEditableMesh* EditableM
 
 void UMeshEditorStaticMeshAdapter::OnDeleteEdges( const UEditableMesh* EditableMesh, const TArray<FEdgeID>& EdgeIDs )
 {
-	const FEdgeArray& Edges = EditableMesh->GetMeshDescription()->Edges();
+	const UMeshDescription* MeshDescription = EditableMesh->GetMeshDescription();
+	check( MeshDescription );
 
 	for( const FEdgeID EdgeID : EdgeIDs )
 	{
-		const FMeshEdge& Edge = Edges[ EdgeID ];
-
-		for( const FPolygonID PolygonID : Edge.ConnectedPolygons )
+		for( const FPolygonID PolygonID : MeshDescription->GetEdgeConnectedPolygons( EdgeID ) )
 		{
 			WireframeMesh->RemoveEdgeInstance( EdgeID, PolygonID );
 		}
@@ -255,11 +263,11 @@ void UMeshEditorStaticMeshAdapter::OnSetEdgesVertices( const UEditableMesh* Edit
 }
 
 
-void UMeshEditorStaticMeshAdapter::OnSetEdgeAttribute( const UEditableMesh* EditableMesh, const FEdgeID EdgeID, const FName AttributeName, const int32 AttributeIndex, const FVector4 AttributeValue )
+void UMeshEditorStaticMeshAdapter::OnSetEdgeAttribute( const UEditableMesh* EditableMesh, const FEdgeID EdgeID, const FMeshElementAttributeData& Attribute )
 {
-	if( AttributeName == UEditableMeshAttribute::EdgeIsHard() )
+	if( Attribute.AttributeName == MeshAttribute::Edge::IsHard )
 	{
-		WireframeMesh->SetEdgeColor( EdgeID, GetEdgeColor( AttributeValue.X > 0.0f ) );
+		WireframeMesh->SetEdgeColor( EdgeID, GetEdgeColor( Attribute.AttributeValue.GetValue<bool>() ) );
 	}
 }
 
diff --git a/Engine/Plugins/Editor/MeshEditor/Source/MeshEditor/MeshEditorStaticMeshAdapter.h b/Engine/Plugins/Editor/MeshEditor/Source/MeshEditor/MeshEditorStaticMeshAdapter.h
--- a/Engine/Plugins/Editor/MeshEditor/Source/MeshEditor/MeshEditorStaticMeshAdapter.h
+++ b/Engine/Plugins/Editor/MeshEditor/Source/MeshEditor/MeshEditorStaticMeshAdapter.h
@@ -42,17 +42,19 @@ public:
 	virtual void OnCreateEmptyVertexRange( const UEditableMesh* EditableMesh, const TArray<FVertexID>& VertexIDs ) override;
 	virtual void OnCreateVertices( const UEditableMesh* EditableMesh, const TArray<FVertexID>& VertexIDs ) override;
 	virtual void OnCreateVertexInstances( const UEditableMesh* EditableMesh, const TArray<FVertexInstanceID>& VertexInstanceIDs ) override;
-	virtual void OnSetVertexAttribute( const UEditableMesh* EditableMesh, const FVertexID VertexID, const FName AttributeName, const int32 AttributeIndex, const FVector4 AttributeValue ) override;
-	virtual void OnSetVertexInstanceAttribute( const UEditableMesh* EditableMesh, const FVertexInstanceID VertexInstanceID, const FName AttributeName, const int32 AttributeIndex, const FVector4 AttributeValue ) override;
+	virtual void OnSetVertexAttribute( const UEditableMesh* EditableMesh, const FVertexID VertexID, const FMeshElementAttributeData& Attribute ) override;
+	virtual void OnSetVertexInstanceAttribute( const UEditableMesh* EditableMesh, const FVertexInstanceID VertexInstanceID, const FMeshElementAttributeData& Attribute ) override {}
 	virtual void OnCreateEdges( const UEditableMesh* EditableMesh, const TArray<FEdgeID>& EdgeIDs ) override;
 	virtual void OnDeleteEdges( const UEditableMesh* EditableMesh, const TArray<FEdgeID>& EdgeIDs ) override;
 	virtual void OnSetEdgesVertices( const UEditableMesh* EditableMesh, const TArray<FEdgeID>& EdgeIDs ) override;
-	virtual void OnSetEdgeAttribute( const UEditableMesh* EditableMesh, const FEdgeID EdgeID, const FName AttributeName, const int32 AttributeIndex, const FVector4 AttributeValue ) override;
+	virtual void OnSetEdgeAttribute( const UEditableMesh* EditableMesh, const FEdgeID EdgeID, const FMeshElementAttributeData& Attribute ) override;
 	virtual void OnCreatePolygons( const UEditableMesh* EditableMesh, const TArray<FPolygonID>& PolygonIDs ) override;
 	virtual void OnDeletePolygons( const UEditableMesh* EditableMesh, const TArray<FPolygonID>& PolygonIDs ) override;
 	virtual void OnChangePolygonVertexInstances( const UEditableMesh* EditableMesh, const TArray<FPolygonID>& PolygonIDs ) override;
+	virtual void OnSetPolygonAttribute( const UEditableMesh* EditableMesh, const FPolygonID PolygonID, const FMeshElementAttributeData& Attribute ) override {}
 	virtual void OnCreatePolygonGroups( const UEditableMesh* EditableMesh, const TArray<FPolygonGroupID>& PolygonGroupIDs ) override;
 	virtual void OnDeletePolygonGroups( const UEditableMesh* EditableMesh, const TArray<FPolygonGroupID>& PolygonGroupIDs ) override;
+	virtual void OnSetPolygonGroupAttribute( const UEditableMesh* EditableMesh, const FPolygonGroupID PolygonGroupID, const FMeshElementAttributeData& Attribute ) override {}
 	virtual void OnAssignPolygonsToPolygonGroups( const UEditableMesh* EditableMesh, const TArray<FPolygonGroupForPolygon>& PolygonGroupForPolygons ) override {}
 	virtual void OnRetriangulatePolygons( const UEditableMesh* EditableMesh, const TArray<FPolygonID>& PolygonIDs ) override;
 
diff --git a/Engine/Plugins/Editor/MeshEditor/Source/MeshEditor/MeshEditorSubdividedStaticMeshAdapter.cpp b/Engine/Plugins/Editor/MeshEditor/Source/MeshEditor/MeshEditorSubdividedStaticMeshAdapter.cpp
--- a/Engine/Plugins/Editor/MeshEditor/Source/MeshEditor/MeshEditorSubdividedStaticMeshAdapter.cpp
+++ b/Engine/Plugins/Editor/MeshEditor/Source/MeshEditor/MeshEditorSubdividedStaticMeshAdapter.cpp
@@ -176,16 +176,6 @@ void UMeshEditorSubdividedStaticMeshAdapter::OnCreateVertexInstances( const UEdi
 }
 
 
-void UMeshEditorSubdividedStaticMeshAdapter::OnSetVertexAttribute( const UEditableMesh* EditableMesh, const FVertexID VertexID, const FName AttributeName, const int32 AttributeIndex, const FVector4 AttributeValue )
-{
-}
-
-
-void UMeshEditorSubdividedStaticMeshAdapter::OnSetVertexInstanceAttribute( const UEditableMesh* EditableMesh, const FVertexInstanceID VertexInstanceID, const FName AttributeName, const int32 AttributeIndex, const FVector4 AttributeValue )
-{
-}
-
-
 void UMeshEditorSubdividedStaticMeshAdapter::OnCreateEdges( const UEditableMesh* EditableMesh, const TArray<FEdgeID>& EdgeIDs )
 {
 }
@@ -201,11 +191,6 @@ void UMeshEditorSubdividedStaticMeshAdapter::OnSetEdgesVertices( const UEditable
 }
 
 
-void UMeshEditorSubdividedStaticMeshAdapter::OnSetEdgeAttribute( const UEditableMesh* EditableMesh, const FEdgeID EdgeID, const FName AttributeName, const int32 AttributeIndex, const FVector4 AttributeValue )
-{
-}
-
-
 void UMeshEditorSubdividedStaticMeshAdapter::OnCreatePolygons( const UEditableMesh* EditableMesh, const TArray<FPolygonID>& PolygonIDs )
 {
 }
diff --git a/Engine/Plugins/Editor/MeshEditor/Source/MeshEditor/MeshEditorSubdividedStaticMeshAdapter.h b/Engine/Plugins/Editor/MeshEditor/Source/MeshEditor/MeshEditorSubdividedStaticMeshAdapter.h
--- a/Engine/Plugins/Editor/MeshEditor/Source/MeshEditor/MeshEditorSubdividedStaticMeshAdapter.h
+++ b/Engine/Plugins/Editor/MeshEditor/Source/MeshEditor/MeshEditorSubdividedStaticMeshAdapter.h
@@ -42,17 +42,19 @@ public:
 	virtual void OnCreateEmptyVertexRange( const UEditableMesh* EditableMesh, const TArray<FVertexID>& VertexIDs ) override;
 	virtual void OnCreateVertices( const UEditableMesh* EditableMesh, const TArray<FVertexID>& VertexIDs ) override;
 	virtual void OnCreateVertexInstances( const UEditableMesh* EditableMesh, const TArray<FVertexInstanceID>& VertexInstanceIDs ) override;
-	virtual void OnSetVertexAttribute( const UEditableMesh* EditableMesh, const FVertexID VertexID, const FName AttributeName, const int32 AttributeIndex, const FVector4 AttributeValue ) override;
-	virtual void OnSetVertexInstanceAttribute( const UEditableMesh* EditableMesh, const FVertexInstanceID VertexInstanceID, const FName AttributeName, const int32 AttributeIndex, const FVector4 AttributeValue ) override;
+	virtual void OnSetVertexAttribute( const UEditableMesh* EditableMesh, const FVertexID VertexID, const FMeshElementAttributeData& Attribute ) override {}
+	virtual void OnSetVertexInstanceAttribute( const UEditableMesh* EditableMesh, const FVertexInstanceID VertexInstanceID, const FMeshElementAttributeData& Attribute ) override {}
 	virtual void OnCreateEdges( const UEditableMesh* EditableMesh, const TArray<FEdgeID>& EdgeIDs ) override;
 	virtual void OnDeleteEdges( const UEditableMesh* EditableMesh, const TArray<FEdgeID>& EdgeIDs ) override;
 	virtual void OnSetEdgesVertices( const UEditableMesh* EditableMesh, const TArray<FEdgeID>& EdgeIDs ) override;
-	virtual void OnSetEdgeAttribute( const UEditableMesh* EditableMesh, const FEdgeID EdgeID, const FName AttributeName, const int32 AttributeIndex, const FVector4 AttributeValue ) override;
+	virtual void OnSetEdgeAttribute( const UEditableMesh* EditableMesh, const FEdgeID EdgeID, const FMeshElementAttributeData& Attribute ) override {}
 	virtual void OnCreatePolygons( const UEditableMesh* EditableMesh, const TArray<FPolygonID>& PolygonIDs ) override;
 	virtual void OnDeletePolygons( const UEditableMesh* EditableMesh, const TArray<FPolygonID>& PolygonIDs ) override;
 	virtual void OnChangePolygonVertexInstances( const UEditableMesh* EditableMesh, const TArray<FPolygonID>& PolygonIDs ) override;
+	virtual void OnSetPolygonAttribute( const UEditableMesh* EditableMesh, const FPolygonID PolygonID, const FMeshElementAttributeData& Attribute ) override {}
 	virtual void OnCreatePolygonGroups( const UEditableMesh* EditableMesh, const TArray<FPolygonGroupID>& PolygonGroupIDs ) override;
 	virtual void OnDeletePolygonGroups( const UEditableMesh* EditableMesh, const TArray<FPolygonGroupID>& PolygonGroupIDs ) override;
+	virtual void OnSetPolygonGroupAttribute( const UEditableMesh* EditableMesh, const FPolygonGroupID PolygonGroupID, const FMeshElementAttributeData& Attribute ) override {}
 	virtual void OnAssignPolygonsToPolygonGroups( const UEditableMesh* EditableMesh, const TArray<FPolygonGroupForPolygon>& PolygonGroupForPolygons ) override {}
 	virtual void OnRetriangulatePolygons( const UEditableMesh* EditableMesh, const TArray<FPolygonID>& PolygonIDs ) override;
 
diff --git a/Engine/Plugins/Editor/MeshEditor/Source/MeshEditor/MeshElementViewportTransformable.cpp b/Engine/Plugins/Editor/MeshEditor/Source/MeshEditor/MeshElementViewportTransformable.cpp
--- a/Engine/Plugins/Editor/MeshEditor/Source/MeshEditor/MeshElementViewportTransformable.cpp
+++ b/Engine/Plugins/Editor/MeshEditor/Source/MeshEditor/MeshElementViewportTransformable.cpp
@@ -3,6 +3,7 @@
 #include "MeshElementViewportTransformable.h"
 #include "MeshEditorMode.h"
 #include "EditableMesh.h"
+#include "MeshAttributes.h"
 
 
 const FTransform FMeshElementViewportTransformable::GetTransform() const
@@ -39,13 +40,14 @@ FBox FMeshElementViewportTransformable::BuildBoundingBox( const FTransform& Boun
 		{
 			if( FMeshEditorMode::IsElementIDValid( MeshElement, EditableMesh ) )
 			{
+				const TVertexAttributeArray<FVector>& VertexPositions = EditableMesh->GetMeshDescription()->VertexAttributes().GetAttributes<FVector>( MeshAttribute::Vertex::Position );
 				BoundingBox.Init();
 
 				FTransform ElementTransform = FTransform::Identity;
 				switch( MeshElement.ElementAddress.ElementType )
 				{
 					case EEditableMeshElementType::Vertex:
-						BoundingBox += ComponentToBoundingBox.TransformPosition( EditableMesh->GetVertexAttribute( FVertexID( MeshElement.ElementAddress.ElementID ), UEditableMeshAttribute::VertexPosition(), 0 ) );
+						BoundingBox += ComponentToBoundingBox.TransformPosition( VertexPositions[ FVertexID( MeshElement.ElementAddress.ElementID ) ] );
 						break;
 
 					case EEditableMeshElementType::Edge:
@@ -53,8 +55,8 @@ FBox FMeshElementViewportTransformable::BuildBoundingBox( const FTransform& Boun
 						FVertexID EdgeVertexID0, EdgeVertexID1;
 						EditableMesh->GetEdgeVertices( FEdgeID( MeshElement.ElementAddress.ElementID ), /* Out */ EdgeVertexID0, /* Out */ EdgeVertexID1 );
 
-						BoundingBox += ComponentToBoundingBox.TransformPosition( EditableMesh->GetVertexAttribute( EdgeVertexID0, UEditableMeshAttribute::VertexPosition(), 0 ) );
-						BoundingBox += ComponentToBoundingBox.TransformPosition( EditableMesh->GetVertexAttribute( EdgeVertexID1, UEditableMeshAttribute::VertexPosition(), 0 ) );
+						BoundingBox += ComponentToBoundingBox.TransformPosition( VertexPositions[ EdgeVertexID0 ] );
+						BoundingBox += ComponentToBoundingBox.TransformPosition( VertexPositions[ EdgeVertexID1 ] );
 					}
 					break;
 
@@ -67,7 +69,7 @@ FBox FMeshElementViewportTransformable::BuildBoundingBox( const FTransform& Boun
 
 						for( const FVertexID VertexID : PerimeterVertexIDs )
 						{
-							BoundingBox += ComponentToBoundingBox.TransformPosition( EditableMesh->GetVertexAttribute( VertexID, UEditableMeshAttribute::VertexPosition(), 0 ) );
+							BoundingBox += ComponentToBoundingBox.TransformPosition( VertexPositions[ VertexID ] );
 						}
 					}
 					break;
diff --git a/Engine/Plugins/Editor/MeshEditor/Source/PolygonModeling/BevelOrInsetPolygon.cpp b/Engine/Plugins/Editor/MeshEditor/Source/PolygonModeling/BevelOrInsetPolygon.cpp
--- a/Engine/Plugins/Editor/MeshEditor/Source/PolygonModeling/BevelOrInsetPolygon.cpp
+++ b/Engine/Plugins/Editor/MeshEditor/Source/PolygonModeling/BevelOrInsetPolygon.cpp
@@ -5,6 +5,7 @@
 #include "IMeshEditorModeUIContract.h"
 #include "UICommandInfo.h"
 #include "EditableMesh.h"
+#include "MeshAttributes.h"
 #include "MeshElement.h"
 #include "MultiBoxBuilder.h"
 #include "UICommandList.h"
@@ -66,6 +67,8 @@ namespace BevelOrInsetPolygonHelpers
 				float CenterWeight = 0.0f;
 				FVector BestEdgeVertex0Position, BestEdgeVertex1Position;
 				{
+					const TVertexAttributeArray<FVector>& VertexPositions = EditableMesh->GetMeshDescription()->VertexAttributes().GetAttributes<FVector>( MeshAttribute::Vertex::Position );
+
 					static TArray<FVertexID> PerimeterVertexIDs;
 					EditableMesh->GetPolygonPerimeterVertices( PolygonID, /* Out */ PerimeterVertexIDs );
 
@@ -77,8 +80,8 @@ namespace BevelOrInsetPolygonHelpers
 						const FVertexID EdgeVertex0 = PerimeterVertexIDs[ VertexNumber ];
 						const FVertexID EdgeVertex1 = PerimeterVertexIDs[ NextVertexNumber ];
 
-						const FVector EdgeVertex0Position = EditableMesh->GetVertexAttribute( EdgeVertex0, UEditableMeshAttribute::VertexPosition(), 0 );
-						const FVector EdgeVertex1Position = EditableMesh->GetVertexAttribute( EdgeVertex1, UEditableMeshAttribute::VertexPosition(), 0 );
+						const FVector EdgeVertex0Position = VertexPositions[ EdgeVertex0 ];
+						const FVector EdgeVertex1Position = VertexPositions[ EdgeVertex1 ];
 
 						BestEdgeVertex0Position = EdgeVertex0Position;
 						BestEdgeVertex1Position = EdgeVertex1Position;
diff --git a/Engine/Plugins/Editor/MeshEditor/Source/PolygonModeling/EditVertexOrEdgeSharpness.cpp b/Engine/Plugins/Editor/MeshEditor/Source/PolygonModeling/EditVertexOrEdgeSharpness.cpp
--- a/Engine/Plugins/Editor/MeshEditor/Source/PolygonModeling/EditVertexOrEdgeSharpness.cpp
+++ b/Engine/Plugins/Editor/MeshEditor/Source/PolygonModeling/EditVertexOrEdgeSharpness.cpp
@@ -5,6 +5,7 @@
 #include "IMeshEditorModeUIContract.h"
 #include "UICommandInfo.h"
 #include "EditableMesh.h"
+#include "MeshAttributes.h"
 #include "MeshElement.h"
 #include "MultiBoxBuilder.h"
 #include "UICommandList.h"
@@ -112,11 +113,13 @@ void UEditVertexCornerSharpnessCommand::ApplyDuringDrag( IMeshEditorModeEditingC
 			static TArray<float> NewSharpnessValues;
 			NewSharpnessValues.Reset();
 
+			const TVertexAttributeArray<float>& VertexSharpnesses = EditableMesh->GetMeshDescription()->VertexAttributes().GetAttributes<float>( MeshAttribute::Vertex::CornerSharpness );
+
 			for( const FMeshElement& VertexElement : VertexElements )
 			{
 				const FVertexID VertexID( VertexElement.ElementAddress.ElementID );
 
-				const float CurrentSharpnessValue = EditableMesh->GetVertexAttribute( VertexID, UEditableMeshAttribute::VertexCornerSharpness(), 0 ).X;
+				const float CurrentSharpnessValue = VertexSharpnesses[ VertexID ];
 				const float NewSharpnessValue = FMath::Clamp( CurrentSharpnessValue + ScaledDragDelta, 0.0f, 1.0f );
 
 				VertexIDs.Add( VertexID );
@@ -161,6 +164,8 @@ void UEditEdgeCreaseSharpnessCommand::ApplyDuringDrag( IMeshEditorModeEditingCon
 			UEditableMesh* EditableMesh = MeshAndSelectedEdges.Key;
 			const TArray<FMeshElement>& EdgeElements = MeshAndSelectedEdges.Value;
 
+			const TEdgeAttributeArray<float>& EdgeSharpnesses = EditableMesh->GetMeshDescription()->EdgeAttributes().GetAttributes<float>( MeshAttribute::Edge::CreaseSharpness );
+
 			static TArray<FEdgeID> EdgeIDs;
 			EdgeIDs.Reset();
 
@@ -171,7 +176,7 @@ void UEditEdgeCreaseSharpnessCommand::ApplyDuringDrag( IMeshEditorModeEditingCon
 			{
 				const FEdgeID EdgeID( EdgeElement.ElementAddress.ElementID );
 
-				const float CurrentSharpnessValue = EditableMesh->GetEdgeAttribute( EdgeID, UEditableMeshAttribute::EdgeCreaseSharpness(), 0 ).X;
+				const float CurrentSharpnessValue = EdgeSharpnesses[ EdgeID ];
 				const float NewSharpnessValue = FMath::Clamp( CurrentSharpnessValue + ScaledDragDelta, 0.0f, 1.0f );
 
 				EdgeIDs.Add( EdgeID );
diff --git a/Engine/Plugins/Editor/MeshEditor/Source/PolygonModeling/SplitPolygon.cpp b/Engine/Plugins/Editor/MeshEditor/Source/PolygonModeling/SplitPolygon.cpp
--- a/Engine/Plugins/Editor/MeshEditor/Source/PolygonModeling/SplitPolygon.cpp
+++ b/Engine/Plugins/Editor/MeshEditor/Source/PolygonModeling/SplitPolygon.cpp
@@ -5,6 +5,7 @@
 #include "IMeshEditorModeUIContract.h"
 #include "UICommandInfo.h"
 #include "EditableMesh.h"
+#include "MeshAttributes.h"
 #include "MeshElement.h"
 #include "MultiBoxBuilder.h"
 #include "UICommandList.h"
@@ -133,6 +134,8 @@ void USplitPolygonCommand::ApplyDuringDrag( IMeshEditorModeEditingContract& Mesh
 			FVector LaserPointerStart, LaserPointerEnd;
 			if( ViewportInteractor->GetLaserPointer( /* Out */ LaserPointerStart, /* Out */ LaserPointerEnd ) )
 			{
+				const TVertexAttributeArray<FVector>& VertexPositions = EditableMesh->GetMeshDescription()->VertexAttributes().GetAttributes<FVector>( MeshAttribute::Vertex::Position );
+
 				for( const FPolygonID CandidatePolygonID : CandidatePolygonIDs )
 				{
 					const FTransform ComponentToWorld = Component->GetComponentToWorld();
@@ -140,15 +143,15 @@ void USplitPolygonCommand::ApplyDuringDrag( IMeshEditorModeEditingContract& Mesh
 					FVector SplitStartLocation;
 					if( StartingVertexID != FVertexID::Invalid )
 					{
-						SplitStartLocation = ComponentToWorld.TransformPosition( EditableMesh->GetVertexAttribute( StartingVertexID, UEditableMeshAttribute::VertexPosition(), 0 ) );
+						SplitStartLocation = ComponentToWorld.TransformPosition( VertexPositions[ StartingVertexID ] );
 					}
 					else if( ensure( StartingEdgeID != FEdgeID::Invalid ) )
 					{
 						FVertexID EdgeVertex0, EdgeVertex1;
 						EditableMesh->GetEdgeVertices( StartingEdgeID, /* Out */ EdgeVertex0, /* Out */ EdgeVertex1 );
 
-						const FVector EdgeVertex0Location = ComponentToWorld.TransformPosition( EditableMesh->GetVertexAttribute( EdgeVertex0, UEditableMeshAttribute::VertexPosition(), 0 ) );
-						const FVector EdgeVertex1Location = ComponentToWorld.TransformPosition( EditableMesh->GetVertexAttribute( EdgeVertex1, UEditableMeshAttribute::VertexPosition(), 0 ) ) ;
+						const FVector EdgeVertex0Location = ComponentToWorld.TransformPosition( VertexPositions[ EdgeVertex0 ] );
+						const FVector EdgeVertex1Location = ComponentToWorld.TransformPosition( VertexPositions[ EdgeVertex1 ] ) ;
 
 						SplitStartLocation = FMath::Lerp( EdgeVertex0Location, EdgeVertex1Location, EdgeSplit );
 					}
@@ -197,8 +200,8 @@ void USplitPolygonCommand::ApplyDuringDrag( IMeshEditorModeEditingContract& Mesh
 
 							FVertexID EdgeVertex0, EdgeVertex1;
 							EditableMesh->GetEdgeVertices( TargetEdgeID, /* Out */ EdgeVertex0, /* Out */ EdgeVertex1 );
-							const FVector EdgeVertex0Location = ComponentToWorld.TransformPosition( EditableMesh->GetVertexAttribute( EdgeVertex0, UEditableMeshAttribute::VertexPosition(), 0 ) );
-							const FVector EdgeVertex1Location = ComponentToWorld.TransformPosition( EditableMesh->GetVertexAttribute( EdgeVertex1, UEditableMeshAttribute::VertexPosition(), 0 ) );
+							const FVector EdgeVertex0Location = ComponentToWorld.TransformPosition( VertexPositions[ EdgeVertex0 ] );
+							const FVector EdgeVertex1Location = ComponentToWorld.TransformPosition( VertexPositions[ EdgeVertex1 ] );
 
 							if( !bIsDisqualified )
 							{
@@ -248,8 +251,8 @@ void USplitPolygonCommand::ApplyDuringDrag( IMeshEditorModeEditingContract& Mesh
 							FVertexID EdgeVertex0, EdgeVertex1;
 							EditableMesh->GetEdgeVertices( ClosestEdgeID, /* Out */ EdgeVertex0, /* Out */ EdgeVertex1 );
 
-							const FVector EdgeVertex0Location = ComponentToWorld.TransformPosition( EditableMesh->GetVertexAttribute( EdgeVertex0, UEditableMeshAttribute::VertexPosition(), 0 ) );
-							const FVector EdgeVertex1Location = ComponentToWorld.TransformPosition( EditableMesh->GetVertexAttribute( EdgeVertex1, UEditableMeshAttribute::VertexPosition(), 0 ) );
+							const FVector EdgeVertex0Location = ComponentToWorld.TransformPosition( VertexPositions[ EdgeVertex0 ] );
+							const FVector EdgeVertex1Location = ComponentToWorld.TransformPosition( VertexPositions[ EdgeVertex1 ] );
 
 							const FVector ImpactOnEdge = SplitStartLocation + SplitDirection * ClosestEdgeDistance;
 
@@ -302,7 +305,7 @@ void USplitPolygonCommand::ApplyDuringDrag( IMeshEditorModeEditingContract& Mesh
 								if( !bIsDisqualified )
 								{
 									// Don't allow connecting to vertices that are colinear to the split direction
-									const FVector VertexLocation = ComponentToWorld.TransformPosition( EditableMesh->GetVertexAttribute( TargetVertexID, UEditableMeshAttribute::VertexPosition(), 0 ) );
+									const FVector VertexLocation = ComponentToWorld.TransformPosition( VertexPositions[ TargetVertexID ] );
 									const FVector VertexDirection = ( VertexLocation - SplitStartLocation ).GetSafeNormal();
 									if( FMath::IsNearlyZero( VertexDirection.SizeSquared() ) || FMath::IsNearlyEqual( FMath::Abs( FVector::DotProduct( SplitDirection, VertexDirection ) ), 1.0f ) )
 									{
