Fixed crash on shutdown after editing an editable mesh.

diff --git a/Engine/Plugins/Editor/MeshEditor/Source/MeshEditor/WireframeMeshComponent.cpp b/Engine/Plugins/Editor/MeshEditor/Source/MeshEditor/WireframeMeshComponent.cpp
--- a/Engine/Plugins/Editor/MeshEditor/Source/MeshEditor/WireframeMeshComponent.cpp
+++ b/Engine/Plugins/Editor/MeshEditor/Source/MeshEditor/WireframeMeshComponent.cpp
@@ -15,6 +15,17 @@
 #include "Engine/Engine.h"
 
 
+void UWireframeMesh::BeginDestroy()
+{
+	Super::BeginDestroy();
+
+	if( !HasAnyFlags( RF_ClassDefaultObject ) )
+	{
+		ReleaseResources();
+	}
+}
+
+
 void UWireframeMesh::Reset()
 {
 	Vertices.Reset();
diff --git a/Engine/Plugins/Editor/MeshEditor/Source/MeshEditor/WireframeMeshComponent.h b/Engine/Plugins/Editor/MeshEditor/Source/MeshEditor/WireframeMeshComponent.h
--- a/Engine/Plugins/Editor/MeshEditor/Source/MeshEditor/WireframeMeshComponent.h
+++ b/Engine/Plugins/Editor/MeshEditor/Source/MeshEditor/WireframeMeshComponent.h
@@ -117,6 +117,8 @@ public:
 		: bBoundsDirty( true )
 	{}
 
+	virtual void BeginDestroy();
+
 	void Reset();
 	void AddVertex( const FVertexID VertexID );
 	void SetVertexPosition( const FVertexID VertexID, const FVector& Position );
