Fixed various bugs with wireframe overlays.
- No longer crashes when deleting a mesh which has had a wireframe overlay generated
- Removes wireframe overlays and destroys the associated components when leaving mesh editor mode
- Moved various callback hooks from constructor/destructor to Enter()/Exit()

diff --git a/Engine/Plugins/Editor/MeshEditor/Source/MeshEditor/MeshEditorMode.cpp b/Engine/Plugins/Editor/MeshEditor/Source/MeshEditor/MeshEditorMode.cpp
--- a/Engine/Plugins/Editor/MeshEditor/Source/MeshEditor/MeshEditorMode.cpp
+++ b/Engine/Plugins/Editor/MeshEditor/Source/MeshEditor/MeshEditorMode.cpp
@@ -453,13 +453,6 @@ FMeshEditorMode::FMeshEditorMode()
 
 	// Register UI commands
 	BindCommands();
-
-	// Notify when the map changes
-	FLevelEditorModule& LevelEditor = FModuleManager::LoadModuleChecked<FLevelEditorModule>( "LevelEditor" );
-	LevelEditor.OnMapChanged().AddRaw( this, &FMeshEditorMode::OnMapChanged );
-	LevelEditor.OnActorSelectionChanged().AddRaw( this, &FMeshEditorMode::OnActorSelectionChanged );
-
-	FEditorDelegates::EndPIE.AddRaw( this, &FMeshEditorMode::OnEndPIE );
 }
 
 
@@ -473,15 +466,6 @@ void FMeshEditorMode::Initialize()
 
 FMeshEditorMode::~FMeshEditorMode()
 {
-	FEditorDelegates::EndPIE.RemoveAll( this );
-
-	FLevelEditorModule* LevelEditor = FModuleManager::GetModulePtr<FLevelEditorModule>( "LevelEditor" );
-	if( LevelEditor )
-	{
-		LevelEditor->OnActorSelectionChanged().RemoveAll( this );
-		LevelEditor->OnMapChanged().RemoveAll( this );
-	}
-
 	// Unregister mesh editor actions
 	FMeshEditorPolygonCommands::Unregister();
 	FMeshEditorEdgeCommands::Unregister();
@@ -834,6 +818,13 @@ void FMeshEditorMode::Enter()
 	// Call parent implementation
 	FEdMode::Enter();
 
+	// Notify when the map changes
+	FLevelEditorModule& LevelEditorModule = FModuleManager::LoadModuleChecked<FLevelEditorModule>( "LevelEditor" );
+	LevelEditorModule.OnMapChanged().AddRaw( this, &FMeshEditorMode::OnMapChanged );
+	LevelEditorModule.OnActorSelectionChanged().AddRaw( this, &FMeshEditorMode::OnActorSelectionChanged );
+
+	FEditorDelegates::EndPIE.AddRaw( this, &FMeshEditorMode::OnEndPIE );
+
 	// Add overlay component for rendering hovered elements
 	HoveredElementsComponent= NewObject<UOverlayComponent>( WireframeComponentContainer );
 	HoveredElementsComponent->SetLineMaterial( OverlayLineMaterial );
@@ -900,7 +891,7 @@ void FMeshEditorMode::Enter()
 
 	// Set the current viewport.
 	{
-		const TSharedRef< ILevelEditor >& LevelEditor = FModuleManager::GetModuleChecked<FLevelEditorModule>("LevelEditor").GetFirstLevelEditor().ToSharedRef();
+		const TSharedRef<ILevelEditor>& LevelEditor = LevelEditorModule.GetFirstLevelEditor().ToSharedRef();
 
 		// Do we have an active perspective viewport that is valid for VR?  If so, go ahead and use that.
 		TSharedPtr<FEditorViewportClient> ViewportClient;
@@ -973,11 +964,29 @@ void FMeshEditorMode::Exit()
 		GEditor->NoteSelectionChange();
 	}
 
+	// Remove wireframe overlays
+	for( const auto& ComponentAndWireframeComponents : ComponentToWireframeComponentMap )
+	{
+		ComponentAndWireframeComponents.Value.WireframeMeshComponent->DestroyComponent();
+		ComponentAndWireframeComponents.Value.WireframeSubdividedMeshComponent->DestroyComponent();
+	}
+	ComponentToWireframeComponentMap.Empty();
+
 	// Remove overlay components
+	DebugNormalsComponent->DestroyComponent();
 	SelectedSubDElementsComponent->DestroyComponent();
 	SelectedElementsComponent->DestroyComponent();
 	HoveredElementsComponent->DestroyComponent();
 
+	FEditorDelegates::EndPIE.RemoveAll( this );
+
+	FLevelEditorModule* LevelEditor = FModuleManager::GetModulePtr<FLevelEditorModule>( "LevelEditor" );
+	if( LevelEditor )
+	{
+		LevelEditor->OnActorSelectionChanged().RemoveAll( this );
+		LevelEditor->OnMapChanged().RemoveAll( this );
+	}
+
 	// Call parent implementation
 	FEdMode::Exit();
 }
