Fixed a whole bunch of bugs relating to GC in mesh editor mode:
- Changed raw UObject* pointers to TWeakObjectPtrs where appropriate.
- Fixed issue where container actor was being removed from the world when changing maps in Mesh Editor mode, and not replaced.
- Stopped adding cached actor components as referenced objects (this was resulting in a GC fail when destroying the world).
- ComponentToWireframeComponentMap is now keyed on FObjectKey, this being the only safe type to use when keying on a UObject.
- Tidies up GC'd entries in the caches when selection changes.

diff --git a/Engine/Plugins/Editor/MeshEditor/Source/MeshEditor/MeshEditorMode.cpp b/Engine/Plugins/Editor/MeshEditor/Source/MeshEditor/MeshEditorMode.cpp
--- a/Engine/Plugins/Editor/MeshEditor/Source/MeshEditor/MeshEditorMode.cpp
+++ b/Engine/Plugins/Editor/MeshEditor/Source/MeshEditor/MeshEditorMode.cpp
@@ -492,6 +492,14 @@ void FMeshEditorMode::OnMapChanged( UWorld* World, EMapChangeType MapChangeType
 	if( MapChangeType == EMapChangeType::TearDownWorld )
 	{
 		RemoveEditableMeshReferences();
+		WireframeComponentContainer = nullptr;
+	}
+	else if( MapChangeType == EMapChangeType::LoadMap || MapChangeType == EMapChangeType::NewMap )
+	{
+		// New world, new component container actor
+		FActorSpawnParameters ActorSpawnParameters;
+		ActorSpawnParameters.ObjectFlags |= RF_Transient;
+		WireframeComponentContainer = GetWorld()->SpawnActor<AActor>( ActorSpawnParameters );
 	}
 }
 
@@ -583,6 +591,14 @@ void FMeshEditorMode::RemoveEditableMeshReferences()
 	PreviewRevertChanges.Empty();
 	ActiveActionModifiedMeshes.Reset();
 
+	// Remove wireframe overlays
+	for( const auto& ComponentAndWireframeComponents : ComponentToWireframeComponentMap )
+	{
+		ComponentAndWireframeComponents.Value.WireframeMeshComponent->DestroyComponent();
+		ComponentAndWireframeComponents.Value.WireframeSubdividedMeshComponent->DestroyComponent();
+	}
+	ComponentToWireframeComponentMap.Empty();
+
 	if( ViewportWorldInteraction != nullptr )
 	{
 		const bool bNewObjectsSelected = false;
@@ -1309,7 +1325,12 @@ void FMeshEditorMode::UpdateDebugNormals()
 
 	for( const FComponentAndEditableMesh& ComponentAndEditableMesh : SelectedComponentsAndEditableMeshes )
 	{
-		const UPrimitiveComponent* Component = ComponentAndEditableMesh.Component;
+		if( !ComponentAndEditableMesh.Component.IsValid() )
+		{
+			continue;
+		}
+
+		const UPrimitiveComponent* Component = ComponentAndEditableMesh.Component.Get();
 		const UEditableMesh* EditableMesh = ComponentAndEditableMesh.EditableMesh;
 		const UMeshDescription* MeshDescription = EditableMesh->GetMeshDescription();
 
@@ -1624,7 +1645,7 @@ void FMeshEditorMode::CommitSelectedMeshes()
 {
 	for( FComponentAndEditableMesh& ComponentAndEditableMesh : SelectedComponentsAndEditableMeshes )
 	{
-		CommitEditableMeshIfNecessary( ComponentAndEditableMesh.EditableMesh, ComponentAndEditableMesh.Component );
+		CommitEditableMeshIfNecessary( ComponentAndEditableMesh.EditableMesh, ComponentAndEditableMesh.Component.Get() );
 	}
 }
 
@@ -2407,7 +2428,6 @@ void FMeshEditorMode::AddReferencedObjects( FReferenceCollector& Collector )
 	Collector.AddReferencedObjects( SelectedEditableMeshes );
 	for( FComponentAndEditableMesh& ComponentAndEditableMesh : SelectedComponentsAndEditableMeshes )
 	{
-		Collector.AddReferencedObject( ComponentAndEditableMesh.Component );
 		Collector.AddReferencedObject( ComponentAndEditableMesh.EditableMesh );
 	}
 
@@ -2418,6 +2438,12 @@ void FMeshEditorMode::AddReferencedObjects( FReferenceCollector& Collector )
 		Collector.AddReferencedObject( Pair.Value.WireframeSubdividedMesh );
 	}
 
+	for( auto Pair : ComponentToWireframeComponentMap )
+	{
+		Collector.AddReferencedObject( Pair.Value.WireframeMeshComponent );
+		Collector.AddReferencedObject( Pair.Value.WireframeSubdividedMeshComponent );
+	}
+
 	for( TTuple<UObject*, TUniquePtr<FChange>>& ObjectAndPreviewRevertChange : PreviewRevertChanges )
 	{
 		UObject* Object = ObjectAndPreviewRevertChange.Get<0>();
@@ -3015,7 +3041,7 @@ void FMeshEditorMode::UpdateActiveAction( const bool bIsActionFinishing )
 			else
 			{
 				// Currently adds new vertices to whichever editable mesh is currently selected
-				Component = SelectedComponentsAndEditableMeshes[ 0 ].Component;
+				Component = SelectedComponentsAndEditableMeshes[ 0 ].Component.Get();
 				EditableMesh = SelectedComponentsAndEditableMeshes[ 0 ].EditableMesh;
 				// @todo mesheditor: allow multiple selected meshes? What should this do?
 
@@ -4711,10 +4737,10 @@ void FMeshEditorMode::RefreshTransformables( const bool bNewObjectsSelected )
 
 const FMeshEditorMode::FWireframeMeshComponents& FMeshEditorMode::CreateWireframeMeshComponents( UPrimitiveComponent* Component )
 {
-	FWireframeMeshComponents* WireframeMeshComponentsPtr = ComponentToWireframeComponentMap.Find( Component );
+	FWireframeMeshComponents* WireframeMeshComponentsPtr = ComponentToWireframeComponentMap.Find( FObjectKey( Component ) );
 	if( !WireframeMeshComponentsPtr )
 	{
-		WireframeMeshComponentsPtr = &ComponentToWireframeComponentMap.Add( Component );
+		WireframeMeshComponentsPtr = &ComponentToWireframeComponentMap.Add( FObjectKey( Component ) );
 
 		const int32 LODIndex = 0;		// @todo mesheditor: We'll want to select an LOD to edit in various different wants (LOD that's visible, or manual user select, etc.)
 		const FEditableMeshSubMeshAddress SubMeshAddress = UEditableMeshFactory::MakeSubmeshAddress( Component, LODIndex );
@@ -4745,10 +4771,10 @@ const FMeshEditorMode::FWireframeMeshComponents& FMeshEditorMode::CreateWirefram
 
 void FMeshEditorMode::DestroyWireframeMeshComponents( UPrimitiveComponent* Component )
 {
-	FWireframeMeshComponents& WireframeMeshComponents = ComponentToWireframeComponentMap.FindChecked( Component );
+	FWireframeMeshComponents& WireframeMeshComponents = ComponentToWireframeComponentMap.FindChecked( FObjectKey( Component ) );
 	WireframeMeshComponents.WireframeMeshComponent->DestroyComponent();
 	WireframeMeshComponents.WireframeSubdividedMeshComponent->DestroyComponent();
-	ComponentToWireframeComponentMap.Remove( Component );
+	ComponentToWireframeComponentMap.Remove( FObjectKey( Component ) );
 }
 
 
@@ -4757,12 +4783,26 @@ void FMeshEditorMode::UpdateSelectedEditableMeshes()
 	static TSet<UPrimitiveComponent*> DeselectedComponents;
 	DeselectedComponents.Reset();
 
+	// Remove wireframe components corresponding to deleted components
+	for( auto It( ComponentToWireframeComponentMap.CreateIterator() ); It; ++It )
+	{
+		if( It->Key.ResolveObjectPtr() == nullptr )
+		{
+			It->Value.WireframeMeshComponent->DestroyComponent();
+			It->Value.WireframeSubdividedMeshComponent->DestroyComponent();
+			It.RemoveCurrent();
+		}
+	}
+
 	// Make a list of components which have just been deselected.
 	// First add all the components which appear in the out-of-date list.
 	// Later, any components which are still selected will be removed from this list.
 	for( const FComponentAndEditableMesh& ComponentAndEditableMesh : SelectedComponentsAndEditableMeshes )
 	{
-		DeselectedComponents.Add( ComponentAndEditableMesh.Component );
+		if( ComponentAndEditableMesh.Component.IsValid() )
+		{
+			DeselectedComponents.Add( ComponentAndEditableMesh.Component.Get() );
+		}
 	}
 
 	SelectedEditableMeshes.Reset();
@@ -4813,12 +4853,12 @@ void FMeshEditorMode::UpdateSelectedEditableMeshes()
 
 	for( const FComponentAndEditableMesh& ComponentAndEditableMesh : SelectedComponentsAndEditableMeshes )
 	{
-		if( DeselectedComponents.Contains( ComponentAndEditableMesh.Component ) )
+		if( DeselectedComponents.Contains( ComponentAndEditableMesh.Component.Get() ) )
 		{
-			DeselectedComponents.Remove( ComponentAndEditableMesh.Component );
+			DeselectedComponents.Remove( ComponentAndEditableMesh.Component.Get() );
 		}
 
-		const FWireframeMeshComponents& OverlayComponents = CreateWireframeMeshComponents( ComponentAndEditableMesh.Component );
+		const FWireframeMeshComponents& OverlayComponents = CreateWireframeMeshComponents( ComponentAndEditableMesh.Component.Get() );
 		const FTransform Transform = ComponentAndEditableMesh.Component->GetComponentTransform();
 		OverlayComponents.WireframeMeshComponent->SetWorldTransform( Transform );
 		OverlayComponents.WireframeSubdividedMeshComponent->SetWorldTransform( Transform );
diff --git a/Engine/Plugins/Editor/MeshEditor/Source/MeshEditor/MeshEditorMode.h b/Engine/Plugins/Editor/MeshEditor/Source/MeshEditor/MeshEditorMode.h
--- a/Engine/Plugins/Editor/MeshEditor/Source/MeshEditor/MeshEditorMode.h
+++ b/Engine/Plugins/Editor/MeshEditor/Source/MeshEditor/MeshEditorMode.h
@@ -18,6 +18,7 @@
 #include "Framework/Commands/UICommandList.h"
 #include "WireframeMeshComponent.h"
 #include "OverlayComponent.h"
+#include "UObject/ObjectKey.h"
 #include "MeshEditorMode.generated.h"
 
 
@@ -535,7 +536,7 @@ protected:
 
 	struct FComponentAndEditableMesh
 	{
-		UPrimitiveComponent* Component;
+		TWeakObjectPtr<UPrimitiveComponent> Component;
 		UEditableMesh* EditableMesh;
 
 		FComponentAndEditableMesh( UPrimitiveComponent* InitComponent, UEditableMesh* InitEditableMesh )
@@ -582,7 +583,7 @@ protected:
 	TMap<FEditableMeshSubMeshAddress, FEditableAndWireframeMeshes> CachedEditableMeshes;
 
 	/** Map mesh components to their wireframe overlays */
-	TMap<UPrimitiveComponent*, FWireframeMeshComponents> ComponentToWireframeComponentMap;
+	TMap<FObjectKey, FWireframeMeshComponents> ComponentToWireframeComponentMap;
 
 	/** Manages saving undo for selected mesh elements while we're dragging them around */
 	FTrackingTransaction TrackingTransaction;
