Fixed const correctness now that TWeakObjectPtr is more strict.

diff --git a/Engine/Plugins/Editor/MeshEditor/Source/MeshEditor/MeshEditorMode.cpp b/Engine/Plugins/Editor/MeshEditor/Source/MeshEditor/MeshEditorMode.cpp
--- a/Engine/Plugins/Editor/MeshEditor/Source/MeshEditor/MeshEditorMode.cpp
+++ b/Engine/Plugins/Editor/MeshEditor/Source/MeshEditor/MeshEditorMode.cpp
@@ -2577,7 +2577,7 @@ FMeshEditorMode::FMeshEditorInteractorData& FMeshEditorMode::GetMeshEditorIntera
 	{
 		FMeshEditorInteractorData& MeshEditorInteractorData = MeshEditorInteractorDatas[ MeshEditorInteractorDataIndex ];
 
-		UViewportInteractor* CurrentViewportInteractor = MeshEditorInteractorData.ViewportInteractor.Get();
+		const UViewportInteractor* CurrentViewportInteractor = MeshEditorInteractorData.ViewportInteractor.Get();
 		if( CurrentViewportInteractor == nullptr )
 		{
 			// Expired
diff --git a/Engine/Plugins/Editor/MeshEditor/Source/MeshEditor/MeshEditorMode.h b/Engine/Plugins/Editor/MeshEditor/Source/MeshEditor/MeshEditorMode.h
--- a/Engine/Plugins/Editor/MeshEditor/Source/MeshEditor/MeshEditorMode.h
+++ b/Engine/Plugins/Editor/MeshEditor/Source/MeshEditor/MeshEditorMode.h
@@ -63,7 +63,7 @@ public:
 	struct FMeshEditorInteractorData
 	{
 		/** The viewport interactor that is this data's counterpart */
-		TWeakObjectPtr<class UViewportInteractor> ViewportInteractor;
+		TWeakObjectPtr<const class UViewportInteractor> ViewportInteractor;
 
 		/** True if we have a valid interaction grabber sphere right now */
 		bool bGrabberSphereIsValid;
