Updated OverlayComponent to use new dynamic mesh API.

diff --git a/Engine/Plugins/Editor/MeshEditor/Source/MeshEditor/OverlayComponent.cpp b/Engine/Plugins/Editor/MeshEditor/Source/MeshEditor/OverlayComponent.cpp
--- a/Engine/Plugins/Editor/MeshEditor/Source/MeshEditor/OverlayComponent.cpp
+++ b/Engine/Plugins/Editor/MeshEditor/Source/MeshEditor/OverlayComponent.cpp
@@ -13,181 +13,106 @@
 #include "DynamicMeshBuilder.h"
 #include "EngineGlobals.h"
 #include "Engine/Engine.h"
-
-
-struct FOverlayVertex
-{
-	FOverlayVertex() {}
-
-	FOverlayVertex( const FVector& InPosition, const FVector2D& InUV, const FVector& InTangentX, const FVector& InTangentZ, const FColor& InColor )
-		: Position( InPosition ),
-		  UV( InUV ),
-		  TangentX( InTangentX ),
-		  TangentZ( InTangentZ ),
-		  Color( InColor )
-	{
-		TangentZ.Vector.W = 255;
-	}
-
-	FVector Position;
-	FVector2D UV;
-	FPackedNormal TangentX;
-	FPackedNormal TangentZ;
-	FColor Color;
-};
-
-
-class FOverlayVertexBuffer : public FVertexBuffer 
-{
-public:
-	TArray<FOverlayVertex> Vertices;
-
-	virtual void InitRHI() override
-	{
-		FRHIResourceCreateInfo CreateInfo;
-		void* VertexBufferData = nullptr;
-		VertexBufferRHI = RHICreateAndLockVertexBuffer( Vertices.Num() * sizeof( FOverlayVertex ), BUF_Static, CreateInfo, VertexBufferData );
-
-		// Copy the vertex data into the vertex buffer.		
-		FMemory::Memcpy( VertexBufferData, Vertices.GetData(), Vertices.Num() * sizeof( FOverlayVertex ) );
-		RHIUnlockVertexBuffer( VertexBufferRHI );
-	}
-
-};
-
-
-class FOverlayIndexBuffer : public FIndexBuffer 
-{
-public:
-	TArray<int32> Indices;
-
-	virtual void InitRHI() override
-	{
-		FRHIResourceCreateInfo CreateInfo;
-		void* Buffer = nullptr;
-		IndexBufferRHI = RHICreateAndLockIndexBuffer( sizeof( int32 ), Indices.Num() * sizeof( int32 ), BUF_Static, CreateInfo, Buffer );
-
-		// Write the indices to the index buffer.		
-		FMemory::Memcpy( Buffer, Indices.GetData(), Indices.Num() * sizeof( int32 ) );
-		RHIUnlockIndexBuffer( IndexBufferRHI );
-	}
-};
-
-
-class FOverlayVertexFactory : public FLocalVertexFactory
-{
-public:
-	FOverlayVertexFactory(ERHIFeatureLevel::Type InFeatureLevel)
-		: FLocalVertexFactory(InFeatureLevel, "FOverlayVertexFactory")
-	{}
-
-	/** Init function that should only be called on render thread. */
-	void Init_RenderThread( const FOverlayVertexBuffer* VertexBuffer )
-	{
-		check( IsInRenderingThread() );
-
-		FDataType NewData;
-		NewData.PositionComponent = STRUCTMEMBER_VERTEXSTREAMCOMPONENT( VertexBuffer, FOverlayVertex, Position, VET_Float3 );
-		NewData.TextureCoordinates.Add(
-			FVertexStreamComponent( VertexBuffer, STRUCT_OFFSET( FOverlayVertex, UV ), sizeof( FOverlayVertex ), VET_Float2 )
-			);
-		NewData.TangentBasisComponents[ 0 ] = STRUCTMEMBER_VERTEXSTREAMCOMPONENT( VertexBuffer, FOverlayVertex, TangentX, VET_PackedNormal );
-		NewData.TangentBasisComponents[ 1 ] = STRUCTMEMBER_VERTEXSTREAMCOMPONENT( VertexBuffer, FOverlayVertex, TangentZ, VET_PackedNormal );
-		NewData.ColorComponent = STRUCTMEMBER_VERTEXSTREAMCOMPONENT( VertexBuffer, FOverlayVertex, Color, VET_Color );
-
-		SetData( NewData );
-	}
-
-	/** Initialization */
-	void Init( const FOverlayVertexBuffer* VertexBuffer )
-	{
-		if ( IsInRenderingThread() )
-		{
-			Init_RenderThread( VertexBuffer );
-		}
-		else
-		{
-			ENQUEUE_UNIQUE_RENDER_COMMAND_TWOPARAMETER(
-				InitOverlayVertexFactory,
-				FOverlayVertexFactory*, VertexFactory, this,
-				const FOverlayVertexBuffer*, VertexBuffer, VertexBuffer,
-				{
-					VertexFactory->Init_RenderThread( VertexBuffer );
-				}
-			);
-		}	
-	}
-};
+#include "StaticMeshResources.h"
+#include "Algo/Accumulate.h"
 
 
 struct FMeshBatchData
 {
-	FMeshBatchData(ERHIFeatureLevel::Type FeatureLevel)
+	FMeshBatchData()
 		: MaterialProxy(nullptr)
-		, VertexBuffer()
-		, IndexBuffer()
-		, VertexFactory(FeatureLevel)
 	{}
 
 	FMaterialRenderProxy* MaterialProxy;
-	FOverlayVertexBuffer VertexBuffer;
-	FOverlayIndexBuffer IndexBuffer;
-	FOverlayVertexFactory VertexFactory;
+	int32 StartIndex;
+	int32 NumPrimitives;
+	int32 MinVertexIndex;
+	int32 MaxVertexIndex;
 };
 
 
-class FOverlaySceneProxy : public FPrimitiveSceneProxy
+class FOverlaySceneProxy final : public FPrimitiveSceneProxy
 {
 public:
 
 	FOverlaySceneProxy( UOverlayComponent* Component )
 		: FPrimitiveSceneProxy( Component ),
-		  MaterialRelevance( Component->GetMaterialRelevance( GetScene().GetFeatureLevel() ) )
+		  MaterialRelevance( Component->GetMaterialRelevance( GetScene().GetFeatureLevel() ) ),
+		  VertexFactory( GetScene().GetFeatureLevel(), "FOverlaySceneProxy" )
 	{
 		ERHIFeatureLevel::Type FeatureLevel = GetScene().GetFeatureLevel();
 		// @todo perf: Render state is marked dirty every time a new element is added, and a new scene proxy will be created.
 		// Perhaps we can instead create the render buffers as BUF_Dynamic and amend them directly on the render thread.
 		// However they will still require a reallocation if the number of elements grows, so this may not be a worthwhile optimization.
 
+		const int32 NumLineVertices = Component->Lines.Num() * 4;
+		const int32 NumLineIndices = Component->Lines.Num() * 6;
+		const int32 NumPointVertices = Component->Points.Num() * 4;
+		const int32 NumPointIndices = Component->Points.Num() * 6;
+		const int32 NumTriangleVertices = Algo::Accumulate( Component->TrianglesByMaterial, 0, []( int32 Acc, const TSparseArray<FOverlayTriangle>& Tris ) { return Acc + Tris.Num() * 3; } );
+		const int32 NumTriangleIndices = Algo::Accumulate( Component->TrianglesByMaterial, 0, []( int32 Acc, const TSparseArray<FOverlayTriangle>& Tris ) { return Acc + Tris.Num() * 3; } );
+
+		const int32 TotalNumVertices = NumLineVertices + NumPointVertices + NumTriangleVertices;
+		const int32 TotalNumIndices = NumLineIndices + NumPointIndices + NumTriangleIndices;
+		const int32 NumTextureCoordinates = 1;
+
+		VertexBuffers.PositionVertexBuffer.Init( TotalNumVertices );
+		VertexBuffers.StaticMeshVertexBuffer.Init( TotalNumVertices, NumTextureCoordinates );
+		VertexBuffers.ColorVertexBuffer.Init( TotalNumVertices );
+		IndexBuffer.Indices.SetNumUninitialized( TotalNumIndices );
+
+		int32 VertexBufferIndex = 0;
+		int32 IndexBufferIndex = 0;
+
 		// Initialize lines.
 		// Lines are represented as two tris of zero thickness.
 		// The vertex normals are used to hold the normalized line direction + a sign determining the direction in which the
 		// material should thicken the polys.
 		if( Component->Lines.Num() > 0 )
 		{
-			MeshBatchDatas.Emplace(FMeshBatchData(FeatureLevel));
+			MeshBatchDatas.Emplace();
 			FMeshBatchData& MeshBatchData = MeshBatchDatas.Last();
-
+			MeshBatchData.MinVertexIndex = VertexBufferIndex;
+			MeshBatchData.MaxVertexIndex = VertexBufferIndex + NumLineVertices - 1;
+			MeshBatchData.StartIndex = IndexBufferIndex;
+			MeshBatchData.NumPrimitives = Component->Lines.Num() * 2;
 			MeshBatchData.MaterialProxy = Component->GetMaterial( 0 )->GetRenderProxy( IsSelected() );
-			MeshBatchData.VertexFactory.Init( &MeshBatchData.VertexBuffer );
-			BeginInitResource( &MeshBatchData.VertexFactory );
-
-			int32 VertexBufferIndex = MeshBatchData.VertexBuffer.Vertices.AddUninitialized( Component->Lines.Num() * 4 );
-			int32 IndexBufferIndex = MeshBatchData.IndexBuffer.Indices.AddUninitialized( Component->Lines.Num() * 6 );
 
 			for( const FOverlayLine& OverlayLine : Component->Lines )
 			{
 				const FVector LineDirection = ( OverlayLine.End - OverlayLine.Start ).GetSafeNormal();
 				const FVector2D UV( OverlayLine.Thickness, 0.0f );
 
-				MeshBatchData.VertexBuffer.Vertices[ VertexBufferIndex + 0 ] = FOverlayVertex( OverlayLine.Start, UV, FVector::ZeroVector, -LineDirection, OverlayLine.Color );
-				MeshBatchData.VertexBuffer.Vertices[ VertexBufferIndex + 1 ] = FOverlayVertex( OverlayLine.End, UV, FVector::ZeroVector, -LineDirection, OverlayLine.Color );
-				MeshBatchData.VertexBuffer.Vertices[ VertexBufferIndex + 2 ] = FOverlayVertex( OverlayLine.End, UV, FVector::ZeroVector, LineDirection, OverlayLine.Color );
-				MeshBatchData.VertexBuffer.Vertices[ VertexBufferIndex + 3 ] = FOverlayVertex( OverlayLine.Start, UV, FVector::ZeroVector, LineDirection, OverlayLine.Color );
-
-				MeshBatchData.IndexBuffer.Indices[ IndexBufferIndex + 0 ] = VertexBufferIndex + 0;
-				MeshBatchData.IndexBuffer.Indices[ IndexBufferIndex + 1 ] = VertexBufferIndex + 1;
-				MeshBatchData.IndexBuffer.Indices[ IndexBufferIndex + 2 ] = VertexBufferIndex + 2;
-				MeshBatchData.IndexBuffer.Indices[ IndexBufferIndex + 3 ] = VertexBufferIndex + 2;
-				MeshBatchData.IndexBuffer.Indices[ IndexBufferIndex + 4 ] = VertexBufferIndex + 3;
-				MeshBatchData.IndexBuffer.Indices[ IndexBufferIndex + 5 ] = VertexBufferIndex + 0;
+				VertexBuffers.PositionVertexBuffer.VertexPosition( VertexBufferIndex + 0 ) = OverlayLine.Start;
+				VertexBuffers.PositionVertexBuffer.VertexPosition( VertexBufferIndex + 1 ) = OverlayLine.End;
+				VertexBuffers.PositionVertexBuffer.VertexPosition( VertexBufferIndex + 2 ) = OverlayLine.End;
+				VertexBuffers.PositionVertexBuffer.VertexPosition( VertexBufferIndex + 3 ) = OverlayLine.Start;
+
+				VertexBuffers.StaticMeshVertexBuffer.SetVertexTangents( VertexBufferIndex + 0, FVector::ZeroVector, FVector::ZeroVector, -LineDirection );
+				VertexBuffers.StaticMeshVertexBuffer.SetVertexTangents( VertexBufferIndex + 1, FVector::ZeroVector, FVector::ZeroVector, -LineDirection );
+				VertexBuffers.StaticMeshVertexBuffer.SetVertexTangents( VertexBufferIndex + 2, FVector::ZeroVector, FVector::ZeroVector, LineDirection );
+				VertexBuffers.StaticMeshVertexBuffer.SetVertexTangents( VertexBufferIndex + 3, FVector::ZeroVector, FVector::ZeroVector, LineDirection );
+
+				VertexBuffers.StaticMeshVertexBuffer.SetVertexUV( VertexBufferIndex + 0, 0, UV );
+				VertexBuffers.StaticMeshVertexBuffer.SetVertexUV( VertexBufferIndex + 1, 0, UV );
+				VertexBuffers.StaticMeshVertexBuffer.SetVertexUV( VertexBufferIndex + 2, 0, UV );
+				VertexBuffers.StaticMeshVertexBuffer.SetVertexUV( VertexBufferIndex + 3, 0, UV );
+
+				VertexBuffers.ColorVertexBuffer.VertexColor( VertexBufferIndex + 0 ) = OverlayLine.Color;
+				VertexBuffers.ColorVertexBuffer.VertexColor( VertexBufferIndex + 1 ) = OverlayLine.Color;
+				VertexBuffers.ColorVertexBuffer.VertexColor( VertexBufferIndex + 2 ) = OverlayLine.Color;
+				VertexBuffers.ColorVertexBuffer.VertexColor( VertexBufferIndex + 3 ) = OverlayLine.Color;
+
+				IndexBuffer.Indices[ IndexBufferIndex + 0 ] = VertexBufferIndex + 0;
+				IndexBuffer.Indices[ IndexBufferIndex + 1 ] = VertexBufferIndex + 1;
+				IndexBuffer.Indices[ IndexBufferIndex + 2 ] = VertexBufferIndex + 2;
+				IndexBuffer.Indices[ IndexBufferIndex + 3 ] = VertexBufferIndex + 2;
+				IndexBuffer.Indices[ IndexBufferIndex + 4 ] = VertexBufferIndex + 3;
+				IndexBuffer.Indices[ IndexBufferIndex + 5 ] = VertexBufferIndex + 0;
 
 				VertexBufferIndex += 4;
 				IndexBufferIndex += 6;
 			}
-
-			BeginInitResource( &MeshBatchData.VertexBuffer );
-			BeginInitResource( &MeshBatchData.IndexBuffer );
 		}
 
 		// Initialize points.
@@ -196,81 +121,117 @@ public:
 		// Size of the point is given by U0.
 		if( Component->Points.Num() > 0 )
 		{
-			MeshBatchDatas.Emplace(FMeshBatchData(FeatureLevel));
+			MeshBatchDatas.Emplace();
 			FMeshBatchData& MeshBatchData = MeshBatchDatas.Last();
-
+			MeshBatchData.MinVertexIndex = VertexBufferIndex;
+			MeshBatchData.MaxVertexIndex = VertexBufferIndex + NumPointVertices - 1;
+			MeshBatchData.StartIndex = IndexBufferIndex;
+			MeshBatchData.NumPrimitives = Component->Points.Num() * 2;
 			MeshBatchData.MaterialProxy = Component->GetMaterial( 1 )->GetRenderProxy( IsSelected() );
-			MeshBatchData.VertexFactory.Init( &MeshBatchData.VertexBuffer );
-			BeginInitResource( &MeshBatchData.VertexFactory );
-
-			int32 VertexBufferIndex = MeshBatchData.VertexBuffer.Vertices.AddUninitialized( Component->Points.Num() * 4 );
-			int32 IndexBufferIndex = MeshBatchData.IndexBuffer.Indices.AddUninitialized( Component->Points.Num() * 6 );
 
 			for( const FOverlayPoint& OverlayPoint : Component->Points )
 			{
 				const FVector2D UV( OverlayPoint.Size, 0.0f );
 
-				MeshBatchData.VertexBuffer.Vertices[ VertexBufferIndex + 0 ] = FOverlayVertex( OverlayPoint.Position, UV, FVector::ZeroVector, FVector( 1.0f, -1.0f, 0.0f ), OverlayPoint.Color );
-				MeshBatchData.VertexBuffer.Vertices[ VertexBufferIndex + 1 ] = FOverlayVertex( OverlayPoint.Position, UV, FVector::ZeroVector, FVector( 1.0f, 1.0f, 0.0f ), OverlayPoint.Color );
-				MeshBatchData.VertexBuffer.Vertices[ VertexBufferIndex + 2 ] = FOverlayVertex( OverlayPoint.Position, UV, FVector::ZeroVector, FVector( -1.0f, 1.0f, 0.0f ), OverlayPoint.Color );
-				MeshBatchData.VertexBuffer.Vertices[ VertexBufferIndex + 3 ] = FOverlayVertex( OverlayPoint.Position, UV, FVector::ZeroVector, FVector( -1.0f, -1.0f, 0.0f ), OverlayPoint.Color );
-
-				MeshBatchData.IndexBuffer.Indices[ IndexBufferIndex + 0 ] = VertexBufferIndex + 0;
-				MeshBatchData.IndexBuffer.Indices[ IndexBufferIndex + 1 ] = VertexBufferIndex + 1;
-				MeshBatchData.IndexBuffer.Indices[ IndexBufferIndex + 2 ] = VertexBufferIndex + 2;
-				MeshBatchData.IndexBuffer.Indices[ IndexBufferIndex + 3 ] = VertexBufferIndex + 2;
-				MeshBatchData.IndexBuffer.Indices[ IndexBufferIndex + 4 ] = VertexBufferIndex + 3;
-				MeshBatchData.IndexBuffer.Indices[ IndexBufferIndex + 5 ] = VertexBufferIndex + 0;
+				VertexBuffers.PositionVertexBuffer.VertexPosition( VertexBufferIndex + 0 ) = OverlayPoint.Position;
+				VertexBuffers.PositionVertexBuffer.VertexPosition( VertexBufferIndex + 1 ) = OverlayPoint.Position;
+				VertexBuffers.PositionVertexBuffer.VertexPosition( VertexBufferIndex + 2 ) = OverlayPoint.Position;
+				VertexBuffers.PositionVertexBuffer.VertexPosition( VertexBufferIndex + 3 ) = OverlayPoint.Position;
+
+				VertexBuffers.StaticMeshVertexBuffer.SetVertexTangents( VertexBufferIndex + 0, FVector::ZeroVector, FVector::ZeroVector, FVector( 1.0f, -1.0f, 0.0f ) );
+				VertexBuffers.StaticMeshVertexBuffer.SetVertexTangents( VertexBufferIndex + 1, FVector::ZeroVector, FVector::ZeroVector, FVector( 1.0f, 1.0f, 0.0f ) );
+				VertexBuffers.StaticMeshVertexBuffer.SetVertexTangents( VertexBufferIndex + 2, FVector::ZeroVector, FVector::ZeroVector, FVector( -1.0f, 1.0f, 0.0f ) );
+				VertexBuffers.StaticMeshVertexBuffer.SetVertexTangents( VertexBufferIndex + 3, FVector::ZeroVector, FVector::ZeroVector, FVector( -1.0f, -1.0f, 0.0f ) );
+
+				VertexBuffers.StaticMeshVertexBuffer.SetVertexUV( VertexBufferIndex + 0, 0, UV );
+				VertexBuffers.StaticMeshVertexBuffer.SetVertexUV( VertexBufferIndex + 1, 0, UV );
+				VertexBuffers.StaticMeshVertexBuffer.SetVertexUV( VertexBufferIndex + 2, 0, UV );
+				VertexBuffers.StaticMeshVertexBuffer.SetVertexUV( VertexBufferIndex + 3, 0, UV );
+
+				VertexBuffers.ColorVertexBuffer.VertexColor( VertexBufferIndex + 0 ) = OverlayPoint.Color;
+				VertexBuffers.ColorVertexBuffer.VertexColor( VertexBufferIndex + 1 ) = OverlayPoint.Color;
+				VertexBuffers.ColorVertexBuffer.VertexColor( VertexBufferIndex + 2 ) = OverlayPoint.Color;
+				VertexBuffers.ColorVertexBuffer.VertexColor( VertexBufferIndex + 3 ) = OverlayPoint.Color;
+
+				IndexBuffer.Indices[ IndexBufferIndex + 0 ] = VertexBufferIndex + 0;
+				IndexBuffer.Indices[ IndexBufferIndex + 1 ] = VertexBufferIndex + 1;
+				IndexBuffer.Indices[ IndexBufferIndex + 2 ] = VertexBufferIndex + 2;
+				IndexBuffer.Indices[ IndexBufferIndex + 3 ] = VertexBufferIndex + 2;
+				IndexBuffer.Indices[ IndexBufferIndex + 4 ] = VertexBufferIndex + 3;
+				IndexBuffer.Indices[ IndexBufferIndex + 5 ] = VertexBufferIndex + 0;
 
 				VertexBufferIndex += 4;
 				IndexBufferIndex += 6;
 			}
-
-			BeginInitResource( &MeshBatchData.VertexBuffer );
-			BeginInitResource( &MeshBatchData.IndexBuffer );
 		}
 
 		// Triangles
 		int32 MaterialIndex = 2;
 		for( const auto& MaterialTriangles : Component->TrianglesByMaterial )
 		{
-			MeshBatchDatas.Emplace(FMeshBatchData(FeatureLevel));
+			MeshBatchDatas.Emplace();
 			FMeshBatchData& MeshBatchData = MeshBatchDatas.Last();
-
+			MeshBatchData.MinVertexIndex = VertexBufferIndex;
+			MeshBatchData.MaxVertexIndex = VertexBufferIndex + MaterialTriangles.Num() * 3 - 1;
+			MeshBatchData.StartIndex = IndexBufferIndex;
+			MeshBatchData.NumPrimitives = MaterialTriangles.Num();
 			MeshBatchData.MaterialProxy = Component->GetMaterial( MaterialIndex )->GetRenderProxy( IsSelected() );
-			MeshBatchData.VertexFactory.Init( &MeshBatchData.VertexBuffer );
-			BeginInitResource( &MeshBatchData.VertexFactory );
-
-			int32 VertexBufferIndex = MeshBatchData.VertexBuffer.Vertices.AddUninitialized( MaterialTriangles.Num() * 3 );
-			int32 IndexBufferIndex = MeshBatchData.IndexBuffer.Indices.AddUninitialized( MaterialTriangles.Num() * 3 );
 
 			for( const FOverlayTriangle& Triangle : MaterialTriangles )
 			{
-				MeshBatchData.VertexBuffer.Vertices[ VertexBufferIndex + 0 ] = FOverlayVertex( Triangle.Vertex0.Position, Triangle.Vertex0.UV, FVector( 1, 0, 0 ), Triangle.Vertex0.Normal, Triangle.Vertex0.Color );
-				MeshBatchData.VertexBuffer.Vertices[ VertexBufferIndex + 1 ] = FOverlayVertex( Triangle.Vertex1.Position, Triangle.Vertex1.UV, FVector( 1, 0, 0 ), Triangle.Vertex1.Normal, Triangle.Vertex1.Color );
-				MeshBatchData.VertexBuffer.Vertices[ VertexBufferIndex + 2 ] = FOverlayVertex( Triangle.Vertex2.Position, Triangle.Vertex2.UV, FVector( 1, 0, 0 ), Triangle.Vertex2.Normal, Triangle.Vertex2.Color );
+				VertexBuffers.PositionVertexBuffer.VertexPosition( VertexBufferIndex + 0 ) = Triangle.Vertex0.Position;
+				VertexBuffers.PositionVertexBuffer.VertexPosition( VertexBufferIndex + 1 ) = Triangle.Vertex1.Position;
+				VertexBuffers.PositionVertexBuffer.VertexPosition( VertexBufferIndex + 2 ) = Triangle.Vertex2.Position;
+
+				VertexBuffers.StaticMeshVertexBuffer.SetVertexTangents( VertexBufferIndex + 0, FVector( 1, 0, 0 ), FVector( 0, 1, 0 ), Triangle.Vertex0.Normal );
+				VertexBuffers.StaticMeshVertexBuffer.SetVertexTangents( VertexBufferIndex + 1, FVector( 1, 0, 0 ), FVector( 0, 1, 0 ), Triangle.Vertex1.Normal );
+				VertexBuffers.StaticMeshVertexBuffer.SetVertexTangents( VertexBufferIndex + 2, FVector( 1, 0, 0 ), FVector( 0, 1, 0 ), Triangle.Vertex2.Normal );
 
-				MeshBatchData.IndexBuffer.Indices[ IndexBufferIndex + 0 ] = VertexBufferIndex + 0;
-				MeshBatchData.IndexBuffer.Indices[ IndexBufferIndex + 1 ] = VertexBufferIndex + 1;
-				MeshBatchData.IndexBuffer.Indices[ IndexBufferIndex + 2 ] = VertexBufferIndex + 2;
+				VertexBuffers.StaticMeshVertexBuffer.SetVertexUV( VertexBufferIndex + 0, 0, Triangle.Vertex0.UV );
+				VertexBuffers.StaticMeshVertexBuffer.SetVertexUV( VertexBufferIndex + 1, 0, Triangle.Vertex1.UV );
+				VertexBuffers.StaticMeshVertexBuffer.SetVertexUV( VertexBufferIndex + 2, 0, Triangle.Vertex2.UV );
+
+				VertexBuffers.ColorVertexBuffer.VertexColor( VertexBufferIndex + 0 ) = Triangle.Vertex0.Color;
+				VertexBuffers.ColorVertexBuffer.VertexColor( VertexBufferIndex + 1 ) = Triangle.Vertex1.Color;
+				VertexBuffers.ColorVertexBuffer.VertexColor( VertexBufferIndex + 2 ) = Triangle.Vertex2.Color;
+
+				IndexBuffer.Indices[ IndexBufferIndex + 0 ] = VertexBufferIndex + 0;
+				IndexBuffer.Indices[ IndexBufferIndex + 1 ] = VertexBufferIndex + 1;
+				IndexBuffer.Indices[ IndexBufferIndex + 2 ] = VertexBufferIndex + 2;
 
 				VertexBufferIndex += 3;
 				IndexBufferIndex += 3;
 			}
 
-			BeginInitResource( &MeshBatchData.VertexBuffer );
-			BeginInitResource( &MeshBatchData.IndexBuffer );
+			MaterialIndex++;
 		}
+
+		ENQUEUE_RENDER_COMMAND( OverlayVertexBuffersInit )(
+			[ this ]( FRHICommandListImmediate& RHICmdList )
+			{
+				VertexBuffers.PositionVertexBuffer.InitResource();
+				VertexBuffers.StaticMeshVertexBuffer.InitResource();
+				VertexBuffers.ColorVertexBuffer.InitResource();
+
+				FLocalVertexFactory::FDataType Data;
+				VertexBuffers.PositionVertexBuffer.BindPositionVertexBuffer( &VertexFactory, Data );
+				VertexBuffers.StaticMeshVertexBuffer.BindTangentVertexBuffer( &VertexFactory, Data );
+				VertexBuffers.StaticMeshVertexBuffer.BindTexCoordVertexBuffer( &VertexFactory, Data );
+				VertexBuffers.ColorVertexBuffer.BindColorVertexBuffer( &VertexFactory, Data );
+				VertexFactory.SetData( Data );
+
+				VertexFactory.InitResource();
+				IndexBuffer.InitResource();
+			});
 	}
 
 	virtual ~FOverlaySceneProxy()
 	{
-		for( FMeshBatchData& MeshBatchData : MeshBatchDatas )
-		{
-			MeshBatchData.VertexBuffer.ReleaseResource();
-			MeshBatchData.IndexBuffer.ReleaseResource();
-			MeshBatchData.VertexFactory.ReleaseResource();
-		}
+		VertexBuffers.PositionVertexBuffer.ReleaseResource();
+		VertexBuffers.StaticMeshVertexBuffer.ReleaseResource();
+		VertexBuffers.ColorVertexBuffer.ReleaseResource();
+		IndexBuffer.ReleaseResource();
+		VertexFactory.ReleaseResource();
 	}
 
 	virtual void GetDynamicMeshElements( const TArray<const FSceneView*>& Views, const FSceneViewFamily& ViewFamily, uint32 VisibilityMap, FMeshElementCollector& Collector ) const override
@@ -286,15 +247,15 @@ public:
 					const FSceneView* View = Views[ ViewIndex ];
 					FMeshBatch& Mesh = Collector.AllocateMesh();
 					FMeshBatchElement& BatchElement = Mesh.Elements[ 0 ];
-					BatchElement.IndexBuffer = &MeshBatchData.IndexBuffer;
+					BatchElement.IndexBuffer = &IndexBuffer;
 					Mesh.bWireframe = false;
-					Mesh.VertexFactory = &MeshBatchData.VertexFactory;
+					Mesh.VertexFactory = &VertexFactory;
 					Mesh.MaterialRenderProxy = MeshBatchData.MaterialProxy;
 					BatchElement.PrimitiveUniformBuffer = CreatePrimitiveUniformBufferImmediate( GetLocalToWorld(), GetBounds(), GetLocalBounds(), true, UseEditorDepthTest() );
-					BatchElement.FirstIndex = 0;
-					BatchElement.NumPrimitives = MeshBatchData.IndexBuffer.Indices.Num() / 3;
-					BatchElement.MinVertexIndex = 0;
-					BatchElement.MaxVertexIndex = MeshBatchData.VertexBuffer.Vertices.Num() - 1;
+					BatchElement.FirstIndex = MeshBatchData.StartIndex;
+					BatchElement.NumPrimitives = MeshBatchData.NumPrimitives;
+					BatchElement.MinVertexIndex = MeshBatchData.MinVertexIndex;
+					BatchElement.MaxVertexIndex = MeshBatchData.MaxVertexIndex;
 					Mesh.ReverseCulling = IsLocalToWorldDeterminantNegative();
 					Mesh.Type = PT_TriangleList;
 					Mesh.DepthPriorityGroup = SDPG_World;
@@ -332,10 +293,13 @@ public:
 		static SIZE_T UniquePointer;
 		return reinterpret_cast<SIZE_T>(&UniquePointer);
 	}
-private:
 
+private:
 	TArray<FMeshBatchData> MeshBatchDatas;
 	FMaterialRelevance MaterialRelevance;
+	FLocalVertexFactory VertexFactory;
+	FStaticMeshVertexBuffers VertexBuffers;
+	FDynamicMeshIndexBuffer32 IndexBuffer;
 };
 
 
