Temporarily disabled wireframe mesh component until I have it working with the new dynamic mesh API.

diff --git a/Engine/Plugins/Editor/MeshEditor/Source/MeshEditor/MeshEditorMode.cpp b/Engine/Plugins/Editor/MeshEditor/Source/MeshEditor/MeshEditorMode.cpp
--- a/Engine/Plugins/Editor/MeshEditor/Source/MeshEditor/MeshEditorMode.cpp
+++ b/Engine/Plugins/Editor/MeshEditor/Source/MeshEditor/MeshEditorMode.cpp
@@ -1063,16 +1063,16 @@ UEditableMesh* FMeshEditorMode::FindOrCreateEditableMesh( UPrimitiveComponent& C
 			// Create a wireframe mesh for the base cage
 			UWireframeMesh* WireframeBaseCage = NewObject<UWireframeMesh>();
 
-			UMeshEditorStaticMeshAdapter* WireframeAdapter = NewObject<UMeshEditorStaticMeshAdapter>();
-			EditableMesh->Adapters.Add( WireframeAdapter );
-			WireframeAdapter->Initialize( EditableMesh, WireframeBaseCage );
+//			UMeshEditorStaticMeshAdapter* WireframeAdapter = NewObject<UMeshEditorStaticMeshAdapter>();
+//			EditableMesh->Adapters.Add( WireframeAdapter );
+//			WireframeAdapter->Initialize( EditableMesh, WireframeBaseCage );
 
 			// Create a wireframe mesh for the subdivided mesh
 			UWireframeMesh* WireframeSubdividedMesh = NewObject<UWireframeMesh>();
 
-			UMeshEditorSubdividedStaticMeshAdapter* WireframeSubdividedAdapter = NewObject<UMeshEditorSubdividedStaticMeshAdapter>();
-			EditableMesh->Adapters.Add( WireframeSubdividedAdapter );
-			WireframeSubdividedAdapter->Initialize( EditableMesh, WireframeSubdividedMesh );
+//			UMeshEditorSubdividedStaticMeshAdapter* WireframeSubdividedAdapter = NewObject<UMeshEditorSubdividedStaticMeshAdapter>();
+//			EditableMesh->Adapters.Add( WireframeSubdividedAdapter );
+//			WireframeSubdividedAdapter->Initialize( EditableMesh, WireframeSubdividedMesh );
 
 			// Rebuild mesh so that the wireframe meshes get their render data built through the adapters
 			EditableMesh->RebuildRenderMesh();
@@ -4752,14 +4752,14 @@ const FMeshEditorMode::FWireframeMeshComponents& FMeshEditorMode::CreateWirefram
 		WireframeSubdividedMeshComponent->SetMaterial( 0, SubdividedMeshWireMaterial );
 		WireframeSubdividedMeshComponent->TranslucencySortPriority = 100;
 		WireframeSubdividedMeshComponent->SetWireframeMesh( EditableAndWireframeMeshes.WireframeSubdividedMesh );
-		WireframeSubdividedMeshComponent->RegisterComponent();
+//////		WireframeSubdividedMeshComponent->RegisterComponent();
 
 		// Create the base cage wireframe mesh component
 		UWireframeMeshComponent* WireframeMeshComponent = NewObject<UWireframeMeshComponent>( WireframeComponentContainer );
 		WireframeMeshComponent->SetMaterial( 0, WireMaterial );
 		WireframeMeshComponent->TranslucencySortPriority = 300;
 		WireframeMeshComponent->SetWireframeMesh( EditableAndWireframeMeshes.WireframeBaseCage );
-		WireframeMeshComponent->RegisterComponent();
+//////		WireframeMeshComponent->RegisterComponent();
 
 		WireframeMeshComponentsPtr->WireframeMeshComponent = WireframeMeshComponent;
 		WireframeMeshComponentsPtr->WireframeSubdividedMeshComponent = WireframeSubdividedMeshComponent;
