Fixed up WireframeMeshComponent so that it works with the new dynamic mesh API.

diff --git a/Engine/Plugins/Editor/MeshEditor/Source/MeshEditor/MeshEditorMode.cpp b/Engine/Plugins/Editor/MeshEditor/Source/MeshEditor/MeshEditorMode.cpp
--- a/Engine/Plugins/Editor/MeshEditor/Source/MeshEditor/MeshEditorMode.cpp
+++ b/Engine/Plugins/Editor/MeshEditor/Source/MeshEditor/MeshEditorMode.cpp
@@ -1063,16 +1063,16 @@ UEditableMesh* FMeshEditorMode::FindOrCreateEditableMesh( UPrimitiveComponent& C
 			// Create a wireframe mesh for the base cage
 			UWireframeMesh* WireframeBaseCage = NewObject<UWireframeMesh>();
 
-//			UMeshEditorStaticMeshAdapter* WireframeAdapter = NewObject<UMeshEditorStaticMeshAdapter>();
-//			EditableMesh->Adapters.Add( WireframeAdapter );
-//			WireframeAdapter->Initialize( EditableMesh, WireframeBaseCage );
+			UMeshEditorStaticMeshAdapter* WireframeAdapter = NewObject<UMeshEditorStaticMeshAdapter>();
+			EditableMesh->Adapters.Add( WireframeAdapter );
+			WireframeAdapter->Initialize( EditableMesh, WireframeBaseCage );
 
 			// Create a wireframe mesh for the subdivided mesh
 			UWireframeMesh* WireframeSubdividedMesh = NewObject<UWireframeMesh>();
 
-//			UMeshEditorSubdividedStaticMeshAdapter* WireframeSubdividedAdapter = NewObject<UMeshEditorSubdividedStaticMeshAdapter>();
-//			EditableMesh->Adapters.Add( WireframeSubdividedAdapter );
-//			WireframeSubdividedAdapter->Initialize( EditableMesh, WireframeSubdividedMesh );
+			UMeshEditorSubdividedStaticMeshAdapter* WireframeSubdividedAdapter = NewObject<UMeshEditorSubdividedStaticMeshAdapter>();
+			EditableMesh->Adapters.Add( WireframeSubdividedAdapter );
+			WireframeSubdividedAdapter->Initialize( EditableMesh, WireframeSubdividedMesh );
 
 			// Rebuild mesh so that the wireframe meshes get their render data built through the adapters
 			EditableMesh->RebuildRenderMesh();
@@ -4752,14 +4752,14 @@ const FMeshEditorMode::FWireframeMeshComponents& FMeshEditorMode::CreateWirefram
 		WireframeSubdividedMeshComponent->SetMaterial( 0, SubdividedMeshWireMaterial );
 		WireframeSubdividedMeshComponent->TranslucencySortPriority = 100;
 		WireframeSubdividedMeshComponent->SetWireframeMesh( EditableAndWireframeMeshes.WireframeSubdividedMesh );
-//////		WireframeSubdividedMeshComponent->RegisterComponent();
+		WireframeSubdividedMeshComponent->RegisterComponent();
 
 		// Create the base cage wireframe mesh component
 		UWireframeMeshComponent* WireframeMeshComponent = NewObject<UWireframeMeshComponent>( WireframeComponentContainer );
 		WireframeMeshComponent->SetMaterial( 0, WireMaterial );
 		WireframeMeshComponent->TranslucencySortPriority = 300;
 		WireframeMeshComponent->SetWireframeMesh( EditableAndWireframeMeshes.WireframeBaseCage );
-//////		WireframeMeshComponent->RegisterComponent();
+		WireframeMeshComponent->RegisterComponent();
 
 		WireframeMeshComponentsPtr->WireframeMeshComponent = WireframeMeshComponent;
 		WireframeMeshComponentsPtr->WireframeSubdividedMeshComponent = WireframeSubdividedMeshComponent;
diff --git a/Engine/Plugins/Editor/MeshEditor/Source/MeshEditor/WireframeMeshComponent.cpp b/Engine/Plugins/Editor/MeshEditor/Source/MeshEditor/WireframeMeshComponent.cpp
--- a/Engine/Plugins/Editor/MeshEditor/Source/MeshEditor/WireframeMeshComponent.cpp
+++ b/Engine/Plugins/Editor/MeshEditor/Source/MeshEditor/WireframeMeshComponent.cpp
@@ -10,7 +10,6 @@
 #include "Materials/Material.h"
 #include "LocalVertexFactory.h"
 #include "SceneManagement.h"
-#include "DynamicMeshBuilder.h"
 #include "EngineGlobals.h"
 #include "Engine/Engine.h"
 
@@ -180,11 +179,17 @@ void UWireframeMesh::InitResources()
 {
 	const int32 EdgeInstanceCount = EdgeInstances.Num();
 
-	VertexBuffer.Vertices.Reset();
-	IndexBuffer.Indices.Reset();
-
 	if( EdgeInstanceCount > 0 )
 	{
+		const int32 NumVertices = EdgeInstanceCount * 4;
+		const int32 NumIndices = EdgeInstanceCount * 6;
+		const int32 NumTextureCoordinates = 1;
+
+		VertexBuffers.PositionVertexBuffer.Init( NumVertices );
+		VertexBuffers.StaticMeshVertexBuffer.Init( NumVertices, NumTextureCoordinates );
+		VertexBuffers.ColorVertexBuffer.Init( NumVertices );
+		IndexBuffer.Indices.SetNumUninitialized( NumIndices );
+
 		// An edge instance is a concrete instance of an edge on a particular polygon. Each edge instance is expanded into two triangles which form a camera facing quad.
 		// The quad has a thickness of zero: the material is responsible for giving it finite thickness, according to the edge and the camera direction.
 
@@ -193,8 +198,8 @@ void UWireframeMesh::InitResources()
 		// The vertex tangent is used to represent the edge normal, which is used for optional backface culling.
 		// The UV0 channel is set on a per-component basis, and contains various overrides per edge, to control the opacity and highlighting.
 
-		int32 VertexBufferIndex = VertexBuffer.Vertices.AddUninitialized( EdgeInstanceCount * 4 );
-		int32 IndexBufferIndex = IndexBuffer.Indices.AddUninitialized( EdgeInstanceCount * 6 );
+		int32 VertexBufferIndex = 0;
+		int32 IndexBufferIndex = 0;
 
 		for( const FWireframeEdgeInstance& EdgeInstance : EdgeInstances )
 		{
@@ -205,10 +210,20 @@ void UWireframeMesh::InitResources()
 			const FVector& EndVertex = Vertices[ Edge.EndVertex.GetValue() ].Position;
 			const FVector EdgeDirection = ( EndVertex - StartVertex ).GetSafeNormal();
 
-			VertexBuffer.Vertices[ VertexBufferIndex + 0 ] = FWireframeMeshVertex( StartVertex, PolygonNormal, -EdgeDirection, Edge.Color );
-			VertexBuffer.Vertices[ VertexBufferIndex + 1 ] = FWireframeMeshVertex( EndVertex, PolygonNormal, -EdgeDirection, Edge.Color );
-			VertexBuffer.Vertices[ VertexBufferIndex + 2 ] = FWireframeMeshVertex( EndVertex, PolygonNormal, EdgeDirection, Edge.Color );
-			VertexBuffer.Vertices[ VertexBufferIndex + 3 ] = FWireframeMeshVertex( StartVertex, PolygonNormal, EdgeDirection, Edge.Color );
+			VertexBuffers.PositionVertexBuffer.VertexPosition( VertexBufferIndex + 0 ) = StartVertex;
+			VertexBuffers.PositionVertexBuffer.VertexPosition( VertexBufferIndex + 1 ) = EndVertex;
+			VertexBuffers.PositionVertexBuffer.VertexPosition( VertexBufferIndex + 2 ) = EndVertex;
+			VertexBuffers.PositionVertexBuffer.VertexPosition( VertexBufferIndex + 3 ) = StartVertex;
+
+			VertexBuffers.StaticMeshVertexBuffer.SetVertexTangents( VertexBufferIndex + 0, PolygonNormal, FVector::ZeroVector, -EdgeDirection );
+			VertexBuffers.StaticMeshVertexBuffer.SetVertexTangents( VertexBufferIndex + 1, PolygonNormal, FVector::ZeroVector, -EdgeDirection );
+			VertexBuffers.StaticMeshVertexBuffer.SetVertexTangents( VertexBufferIndex + 2, PolygonNormal, FVector::ZeroVector, EdgeDirection );
+			VertexBuffers.StaticMeshVertexBuffer.SetVertexTangents( VertexBufferIndex + 3, PolygonNormal, FVector::ZeroVector, EdgeDirection );
+
+			VertexBuffers.ColorVertexBuffer.VertexColor( VertexBufferIndex + 0 ) = Edge.Color;
+			VertexBuffers.ColorVertexBuffer.VertexColor( VertexBufferIndex + 1 ) = Edge.Color;
+			VertexBuffers.ColorVertexBuffer.VertexColor( VertexBufferIndex + 2 ) = Edge.Color;
+			VertexBuffers.ColorVertexBuffer.VertexColor( VertexBufferIndex + 3 ) = Edge.Color;
 
 			IndexBuffer.Indices[ IndexBufferIndex + 0 ] = VertexBufferIndex + 0;
 			IndexBuffer.Indices[ IndexBufferIndex + 1 ] = VertexBufferIndex + 1;
@@ -222,7 +237,9 @@ void UWireframeMesh::InitResources()
 		}
 
 		// Enqueue initialization of render resource
-		BeginInitResource( &VertexBuffer );
+		BeginInitResource( &VertexBuffers.PositionVertexBuffer );
+		BeginInitResource( &VertexBuffers.StaticMeshVertexBuffer );
+		BeginInitResource( &VertexBuffers.ColorVertexBuffer );
 		BeginInitResource( &IndexBuffer );
 	}
 }
@@ -232,7 +249,9 @@ void UWireframeMesh::ReleaseResources()
 {
 	if( IndexBuffer.Indices.Num() > 0 )
 	{
-		BeginReleaseResource( &VertexBuffer );
+		BeginReleaseResource( &VertexBuffers.PositionVertexBuffer );
+		BeginReleaseResource( &VertexBuffers.StaticMeshVertexBuffer );
+		BeginReleaseResource( &VertexBuffers.ColorVertexBuffer );
 		BeginReleaseResource( &IndexBuffer );
 	}
 }
@@ -328,20 +347,23 @@ public:
 	FWireframeMeshSceneProxy( UWireframeMeshComponent* Component )
 		: FPrimitiveSceneProxy( Component )
 		, MaterialRelevance( Component->GetMaterialRelevance( GetScene().GetFeatureLevel() ) )
-		, VertexFactory(GetScene().GetFeatureLevel())
+		, VertexFactory( GetScene().GetFeatureLevel(), "FWireframeMeshSceneProxy" )
 		
 	{
 		check( Component->GetWireframeMesh() );
 
-		VertexBuffer = &Component->GetWireframeMesh()->VertexBuffer;
+		const FStaticMeshVertexBuffers* VertexBuffers = &Component->GetWireframeMesh()->VertexBuffers;
 		IndexBuffer = &Component->GetWireframeMesh()->IndexBuffer;
 
-		// Init vertex factory
-		VertexFactory.Init( &Component->GetWireframeMesh()->VertexBuffer, &InstanceVertexBuffer );
-		BeginInitResource( &VertexFactory );
-
 		// Init instance data
-		InstanceVertexBuffer.Vertices.AddZeroed( Component->GetWireframeMesh()->GetNumEdgeInstances() * 4 );
+		NumVertices = Component->GetWireframeMesh()->GetNumEdgeInstances() * 4;
+		const int32 NumTextureCoordinates = 1;
+		InstanceVertexBuffer.Init( NumVertices, NumTextureCoordinates );
+
+		for( int32 Index = 0; Index < NumVertices; ++Index )
+		{
+			InstanceVertexBuffer.SetVertexUV( Index, 0, FVector2D::ZeroVector );
+		}
 
 		for( const FEdgeID HiddenEdgeID : Component->HiddenEdgeIDs )
 		{
@@ -349,14 +371,27 @@ public:
 			for( const int32 EdgeInstanceIndex : EdgeInstanceIndices )
 			{
 				// Hide an edge by setting U of all of its instance vertices to 1.0
-				InstanceVertexBuffer.Vertices[ EdgeInstanceIndex * 4 + 0 ].UV = FVector2D( 1.0f, 0.0f );
-				InstanceVertexBuffer.Vertices[ EdgeInstanceIndex * 4 + 1 ].UV = FVector2D( 1.0f, 0.0f );
-				InstanceVertexBuffer.Vertices[ EdgeInstanceIndex * 4 + 2 ].UV = FVector2D( 1.0f, 0.0f );
-				InstanceVertexBuffer.Vertices[ EdgeInstanceIndex * 4 + 3 ].UV = FVector2D( 1.0f, 0.0f );
+				InstanceVertexBuffer.SetVertexUV( EdgeInstanceIndex * 4 + 0, 0, FVector2D( 1.0f, 0.0f ) );
+				InstanceVertexBuffer.SetVertexUV( EdgeInstanceIndex * 4 + 1, 0, FVector2D( 1.0f, 0.0f ) );
+				InstanceVertexBuffer.SetVertexUV( EdgeInstanceIndex * 4 + 2, 0, FVector2D( 1.0f, 0.0f ) );
+				InstanceVertexBuffer.SetVertexUV( EdgeInstanceIndex * 4 + 3, 0, FVector2D( 1.0f, 0.0f ) );
 			}
 		}
 
-		BeginInitResource( &InstanceVertexBuffer );
+		ENQUEUE_RENDER_COMMAND( WireframeMeshVertexFactoryInit )(
+			[ this, VertexBuffers ]( FRHICommandListImmediate& RHICmdList )
+			{
+				InstanceVertexBuffer.InitResource();
+
+				FLocalVertexFactory::FDataType Data;
+				VertexBuffers->PositionVertexBuffer.BindPositionVertexBuffer( &VertexFactory, Data );
+				VertexBuffers->StaticMeshVertexBuffer.BindTangentVertexBuffer( &VertexFactory, Data );
+				VertexBuffers->ColorVertexBuffer.BindColorVertexBuffer( &VertexFactory, Data );
+				InstanceVertexBuffer.BindTexCoordVertexBuffer( &VertexFactory, Data );
+				VertexFactory.SetData( Data );
+
+				VertexFactory.InitResource();
+			});
 
 		// Grab material
 		Material = Component->GetMaterial( 0 );
@@ -393,7 +428,7 @@ public:
 				BatchElement.FirstIndex = 0;
 				BatchElement.NumPrimitives = IndexBuffer->Indices.Num() / 3;
 				BatchElement.MinVertexIndex = 0;
-				BatchElement.MaxVertexIndex = VertexBuffer->Vertices.Num() - 1;
+				BatchElement.MaxVertexIndex = NumVertices - 1;
 				Mesh.ReverseCulling = IsLocalToWorldDeterminantNegative();
 				Mesh.Type = PT_TriangleList;
 				Mesh.DepthPriorityGroup = SDPG_World;
@@ -433,10 +468,10 @@ public:
 private:
 
 	UMaterialInterface* Material;
-	const FWireframeMeshVertexBuffer* VertexBuffer;
-	const FWireframeMeshIndexBuffer* IndexBuffer;
-	FWireframeMeshInstanceVertexBuffer InstanceVertexBuffer;
-	FWireframeMeshVertexFactory VertexFactory;
+	const FDynamicMeshIndexBuffer32* IndexBuffer;
+	FStaticMeshVertexBuffer InstanceVertexBuffer;
+	FLocalVertexFactory VertexFactory;
+	int32 NumVertices;
 
 	FMaterialRelevance MaterialRelevance;
 };
diff --git a/Engine/Plugins/Editor/MeshEditor/Source/MeshEditor/WireframeMeshComponent.h b/Engine/Plugins/Editor/MeshEditor/Source/MeshEditor/WireframeMeshComponent.h
--- a/Engine/Plugins/Editor/MeshEditor/Source/MeshEditor/WireframeMeshComponent.h
+++ b/Engine/Plugins/Editor/MeshEditor/Source/MeshEditor/WireframeMeshComponent.h
@@ -9,6 +9,8 @@
 #include "RenderResource.h"
 #include "PackedNormal.h"
 #include "EditableMeshTypes.h"
+#include "StaticMeshResources.h"
+#include "DynamicMeshBuilder.h"
 #include "WireframeMeshComponent.generated.h"
 
 class FPrimitiveSceneProxy;
@@ -148,8 +150,8 @@ public:
 	// An edge instance represents a unique quadrilaterial which forms part of the wireframe mesh.
 	TArray<FWireframeEdgeInstance> EdgeInstances;
 
-	FWireframeMeshVertexBuffer VertexBuffer;
-	FWireframeMeshIndexBuffer IndexBuffer;
+	FStaticMeshVertexBuffers VertexBuffers;
+	FDynamicMeshIndexBuffer32 IndexBuffer;
 
 	mutable FBoxSphereBounds Bounds;
 	mutable bool bBoundsDirty;
