Refinement and dead code removal

diff --git a/Engine/Plugins/Editor/MeshEditor/Source/MeshEditor/WireframeMeshComponent.cpp b/Engine/Plugins/Editor/MeshEditor/Source/MeshEditor/WireframeMeshComponent.cpp
--- a/Engine/Plugins/Editor/MeshEditor/Source/MeshEditor/WireframeMeshComponent.cpp
+++ b/Engine/Plugins/Editor/MeshEditor/Source/MeshEditor/WireframeMeshComponent.cpp
@@ -257,88 +257,6 @@ void UWireframeMesh::ReleaseResources()
 }
 
 
-void FWireframeMeshVertexBuffer::InitRHI()
-{
-	FRHIResourceCreateInfo CreateInfo;
-	void* VertexBufferData = nullptr;
-	VertexBufferRHI = RHICreateAndLockVertexBuffer( Vertices.Num() * sizeof( FWireframeMeshVertex ), BUF_Static, CreateInfo, VertexBufferData );
-
-	// Copy the vertex data into the vertex buffer.		
-	FMemory::Memcpy( VertexBufferData, Vertices.GetData(), Vertices.Num() * sizeof( FWireframeMeshVertex ) );
-	RHIUnlockVertexBuffer( VertexBufferRHI );
-}
-
-
-void FWireframeMeshIndexBuffer::InitRHI()
-{
-	FRHIResourceCreateInfo CreateInfo;
-	void* Buffer = nullptr;
-	IndexBufferRHI = RHICreateAndLockIndexBuffer( sizeof( int32 ), Indices.Num() * sizeof( int32 ), BUF_Static, CreateInfo, Buffer );
-
-	// Write the indices to the index buffer.		
-	FMemory::Memcpy( Buffer, Indices.GetData(), Indices.Num() * sizeof( int32 ) );
-	RHIUnlockIndexBuffer( IndexBufferRHI );
-}
-
-
-void FWireframeMeshInstanceVertexBuffer::InitRHI()
-{
-	FRHIResourceCreateInfo CreateInfo;
-	void* VertexBufferData = nullptr;
-	VertexBufferRHI = RHICreateAndLockVertexBuffer( Vertices.Num() * sizeof( FWireframeMeshInstanceVertex ), BUF_Static, CreateInfo, VertexBufferData );
-
-	// Copy the vertex data into the vertex buffer.		
-	FMemory::Memcpy( VertexBufferData, Vertices.GetData(), Vertices.Num() * sizeof( FWireframeMeshInstanceVertex ) );
-	RHIUnlockVertexBuffer( VertexBufferRHI );
-}
-
-
-/** Vertex Factory */
-class FWireframeMeshVertexFactory : public FLocalVertexFactory
-{
-public:
-
-	FWireframeMeshVertexFactory(ERHIFeatureLevel::Type InFeatureLevel)
-		: FLocalVertexFactory(InFeatureLevel, "FWireframeMeshVertexFactory")
-	{}
-
-	/** Init function that should only be called on render thread. */
-	void Init_RenderThread( const FWireframeMeshVertexBuffer* VertexBuffer, const FWireframeMeshInstanceVertexBuffer* InstanceVertexBuffer )
-	{
-		check( IsInRenderingThread() );
-
-		FDataType NewData;
-		NewData.PositionComponent = STRUCTMEMBER_VERTEXSTREAMCOMPONENT( VertexBuffer, FWireframeMeshVertex, Position, VET_Float3 );
-		NewData.TextureCoordinates.Add( FVertexStreamComponent( InstanceVertexBuffer, STRUCT_OFFSET( FWireframeMeshInstanceVertex, UV ), sizeof( FWireframeMeshInstanceVertex ), VET_Float2 ) );
-		NewData.TangentBasisComponents[ 0 ] = STRUCTMEMBER_VERTEXSTREAMCOMPONENT( VertexBuffer, FWireframeMeshVertex, TangentX, VET_PackedNormal );
-		NewData.TangentBasisComponents[ 1 ] = STRUCTMEMBER_VERTEXSTREAMCOMPONENT( VertexBuffer, FWireframeMeshVertex, TangentZ, VET_PackedNormal );
-		NewData.ColorComponent = STRUCTMEMBER_VERTEXSTREAMCOMPONENT( VertexBuffer, FWireframeMeshVertex, Color, VET_Color );
-
-		SetData( NewData );
-	}
-
-	/** Initialization */
-	void Init( const FWireframeMeshVertexBuffer* VertexBuffer, const FWireframeMeshInstanceVertexBuffer* InstanceVertexBuffer )
-	{
-		if( IsInRenderingThread() )
-		{
-			Init_RenderThread( VertexBuffer, InstanceVertexBuffer );
-		}
-		else
-		{
-			ENQUEUE_UNIQUE_RENDER_COMMAND_THREEPARAMETER(
-				InitWireframeMeshVertexFactory,
-				FWireframeMeshVertexFactory*, VertexFactory, this,
-				const FWireframeMeshVertexBuffer*, VertexBuffer, VertexBuffer,
-				const FWireframeMeshInstanceVertexBuffer*, InstanceVertexBuffer, InstanceVertexBuffer,
-				{
-					VertexFactory->Init_RenderThread( VertexBuffer, InstanceVertexBuffer );
-				});
-		}	
-	}
-};
-
-
 /** Scene proxy */
 class FWireframeMeshSceneProxy : public FPrimitiveSceneProxy
 {
diff --git a/Engine/Plugins/Editor/MeshEditor/Source/MeshEditor/WireframeMeshComponent.h b/Engine/Plugins/Editor/MeshEditor/Source/MeshEditor/WireframeMeshComponent.h
--- a/Engine/Plugins/Editor/MeshEditor/Source/MeshEditor/WireframeMeshComponent.h
+++ b/Engine/Plugins/Editor/MeshEditor/Source/MeshEditor/WireframeMeshComponent.h
@@ -16,70 +16,6 @@
 class FPrimitiveSceneProxy;
 
 
-/** Structure representing a vertex definition for the wireframe mesh */
-struct FWireframeMeshVertex
-{
-	FWireframeMeshVertex() {}
-
-	FWireframeMeshVertex( const FVector& InPosition, const FVector& InNormal, const FVector& InTangent, const FColor& InColor )
-		: Position( InPosition ),
-		  TangentX( InNormal ),
-		  TangentZ( InTangent ),
-		  Color( InColor )
-	{
-		TangentZ.Vector.W = 255;
-	}
-
-	FVector Position;
-	FPackedNormal TangentX;
-	FPackedNormal TangentZ;
-	FColor Color;
-};
-
-
-/** Vertex buffer for the shared mesh */
-class FWireframeMeshVertexBuffer : public FVertexBuffer 
-{
-public:
-	TArray<FWireframeMeshVertex> Vertices;
-
-	virtual void InitRHI() override;
-};
-
-
-/** Index buffer for the shared mesh */
-class FWireframeMeshIndexBuffer : public FIndexBuffer 
-{
-public:
-	TArray<int32> Indices;
-
-	virtual void InitRHI() override;
-};
-
-
-/** Structure representing instance-specific extra data for a vertex, passed in the UV0 channel */
-struct FWireframeMeshInstanceVertex
-{
-	FWireframeMeshInstanceVertex() {}
-
-	FWireframeMeshInstanceVertex( const FVector2D& InUV )
-		: UV( InUV )
-	{}
-
-	FVector2D UV;
-};
-
-
-/** Vertex buffer for a mesh instance */
-class FWireframeMeshInstanceVertexBuffer : public FVertexBuffer
-{
-public:
-	TArray<FWireframeMeshInstanceVertex> Vertices;
-
-	virtual void InitRHI() override;
-};
-
-
 struct FWireframeVertex
 {
 	FVector Position;
