PolygonCenter attribute now set correctly in the mesh description.  Changed other code to read from this attribute in preference to calling ComputePolygonCenter, and likewise for getting polygon normals.

diff --git a/Engine/Plugins/Editor/MeshEditor/Source/MeshEditor/MeshEditorMode.cpp b/Engine/Plugins/Editor/MeshEditor/Source/MeshEditor/MeshEditorMode.cpp
--- a/Engine/Plugins/Editor/MeshEditor/Source/MeshEditor/MeshEditorMode.cpp
+++ b/Engine/Plugins/Editor/MeshEditor/Source/MeshEditor/MeshEditorMode.cpp
@@ -3604,11 +3604,13 @@ FEditableMeshElementAddress FMeshEditorMode::QueryElement( const UEditableMesh&
 	static TSet<FPolygonID> FrontFacingPolygons;
 	FrontFacingPolygons.Reset();
 
+	const TPolygonAttributeArray<FVector>& PolygonCenters = EditableMesh.GetMeshDescription()->PolygonAttributes().GetAttributes<FVector>( MeshAttribute::Polygon::Center );
+
 	// Look for all the front-facing elements
 	for( const FPolygonID PolygonID : CandidatePolygons )
 	{
 		const FVector PolygonNormal = EditableMesh.ComputePolygonNormal( PolygonID );
-		const FVector PolygonCenter = EditableMesh.ComputePolygonCenter( PolygonID );
+		const FVector PolygonCenter = PolygonCenters[ PolygonID ];
 		if( ( InteractorShape == EInteractorShape::GrabberSphere ) ||	// Sphere tests never eliminate back-facing geometry
 			!bIsPerspectiveView ||			// @todo mesheditor: Add support for backface culling in orthographic views
 			FVector::DotProduct( CameraLocation - PolygonCenter, PolygonNormal ) > 0.0f )
@@ -4698,7 +4700,8 @@ void FMeshEditorMode::RefreshTransformables( const bool bNewObjectsSelected )
 						{
 							const FPolygonID PolygonID( MeshElement.ElementAddress.ElementID );
 
-							const FVector ComponentSpacePolygonCenter = EditableMesh->ComputePolygonCenter( PolygonID );
+							const TPolygonAttributeArray<FVector>& PolygonCenters = EditableMesh->GetMeshDescription()->PolygonAttributes().GetAttributes<FVector>( MeshAttribute::Polygon::Center );
+							const FVector ComponentSpacePolygonCenter = PolygonCenters[ PolygonID ];
 							ElementTransform.SetLocation( ComponentToWorldMatrix.TransformPosition( ComponentSpacePolygonCenter ) );
 
 							const FVector WindingVector =
diff --git a/Engine/Plugins/Editor/MeshEditor/Source/MeshEditor/MeshEditorStaticMeshAdapter.cpp b/Engine/Plugins/Editor/MeshEditor/Source/MeshEditor/MeshEditorStaticMeshAdapter.cpp
--- a/Engine/Plugins/Editor/MeshEditor/Source/MeshEditor/MeshEditorStaticMeshAdapter.cpp
+++ b/Engine/Plugins/Editor/MeshEditor/Source/MeshEditor/MeshEditorStaticMeshAdapter.cpp
@@ -70,10 +70,11 @@ void UMeshEditorStaticMeshAdapter::OnRebuildRenderMesh( const UEditableMesh* Edi
 
 	// Add all polygons and edge instances
 	// @todo mesheditor wireframe: holes not yet handled
+	const TPolygonAttributeArray<FVector>& PolygonNormals = MeshDescription->PolygonAttributes().GetAttributes<FVector>( MeshAttribute::Polygon::Normal );
 	for( const FPolygonID PolygonID : EditableMesh->GetMeshDescription()->Polygons().GetElementIDs() )
 	{
 		WireframeMesh->AddPolygon( PolygonID );
-		WireframeMesh->SetPolygonNormal( PolygonID, EditableMesh->ComputePolygonNormal( PolygonID ) );
+		WireframeMesh->SetPolygonNormal( PolygonID, PolygonNormals[ PolygonID ] );
 
 		const int32 NumEdges = EditableMesh->GetPolygonPerimeterEdgeCount( PolygonID );
 		for( int32 Index = 0; Index < NumEdges; ++Index )
@@ -274,10 +275,14 @@ void UMeshEditorStaticMeshAdapter::OnSetEdgeAttribute( const UEditableMesh* Edit
 
 void UMeshEditorStaticMeshAdapter::OnCreatePolygons( const UEditableMesh* EditableMesh, const TArray<FPolygonID>& PolygonIDs )
 {
+	const UMeshDescription* MeshDescription = EditableMesh->GetMeshDescription();
+	check( MeshDescription );
+
+	const TPolygonAttributeArray<FVector>& PolygonNormals = MeshDescription->PolygonAttributes().GetAttributes<FVector>( MeshAttribute::Polygon::Normal );
 	for( const FPolygonID PolygonID : PolygonIDs )
 	{
 		WireframeMesh->AddPolygon( PolygonID );
-		WireframeMesh->SetPolygonNormal( PolygonID, EditableMesh->ComputePolygonNormal( PolygonID ) );
+		WireframeMesh->SetPolygonNormal( PolygonID, PolygonNormals[ PolygonID ] );
 
 		const int32 NumEdges = EditableMesh->GetPolygonPerimeterEdgeCount( PolygonID );
 		for( int32 Index = 0; Index < NumEdges; ++Index )
diff --git a/Engine/Plugins/Runtime/EditableMesh/Source/EditableMesh/EditableMesh.cpp b/Engine/Plugins/Runtime/EditableMesh/Source/EditableMesh/EditableMesh.cpp
--- a/Engine/Plugins/Runtime/EditableMesh/Source/EditableMesh/EditableMesh.cpp
+++ b/Engine/Plugins/Runtime/EditableMesh/Source/EditableMesh/EditableMesh.cpp
@@ -6295,9 +6295,19 @@ void UEditableMesh::GeneratePolygonTangentsAndNormals( const TArray<FPolygonID>&
 	TPolygonAttributeArray<FVector>& PolygonNormals = GetMeshDescription()->PolygonAttributes().GetAttributes<FVector>( MeshAttribute::Polygon::Normal );
 	TPolygonAttributeArray<FVector>& PolygonTangents = GetMeshDescription()->PolygonAttributes().GetAttributes<FVector>( MeshAttribute::Polygon::Tangent );
 	TPolygonAttributeArray<FVector>& PolygonBinormals = GetMeshDescription()->PolygonAttributes().GetAttributes<FVector>( MeshAttribute::Polygon::Binormal );
+	TPolygonAttributeArray<FVector>& PolygonCenters = GetMeshDescription()->PolygonAttributes().GetAttributes<FVector>( MeshAttribute::Polygon::Center );
 
 	for( const FPolygonID PolygonID : PolygonIDs )
 	{
+		// Calculate the center of this polygon
+		FVector Center = FVector::ZeroVector;
+		const TArray<FVertexInstanceID>& VertexInstanceIDs = GetMeshDescription()->GetPolygonPerimeterVertexInstances( PolygonID );
+		for( const FVertexInstanceID VertexInstanceID : VertexInstanceIDs )
+		{
+			Center += VertexPositions[ GetMeshDescription()->GetVertexInstanceVertex( VertexInstanceID ) ];
+		}
+		Center /= float( VertexInstanceIDs.Num() );
+
 		// Calculate the tangent basis for the polygon, based on the average of all constituent triangles
 		FVector Normal = FVector::ZeroVector;
 		FVector Tangent = FVector::ZeroVector;
@@ -6330,6 +6340,7 @@ void UEditableMesh::GeneratePolygonTangentsAndNormals( const TArray<FPolygonID>&
 		PolygonNormals[ PolygonID ] = Normal.GetSafeNormal();
 		PolygonTangents[ PolygonID ] = Tangent.GetSafeNormal();
 		PolygonBinormals[ PolygonID ] = Binormal.GetSafeNormal();
+		PolygonCenters[ PolygonID ] = Center;
 	}
 }
 
