Fixed undo/redo issues with applying materials to polygons. There remains a bigger problem which will need a different solution: materials still do not have a key name generated deterministically. This means that sometimes the key will come from the imported material name and at other times from the material itself: this can currently break Editable Mesh.

diff --git a/Engine/Plugins/Editor/MeshEditor/Source/MeshEditor/MeshEditorMode.cpp b/Engine/Plugins/Editor/MeshEditor/Source/MeshEditor/MeshEditorMode.cpp
--- a/Engine/Plugins/Editor/MeshEditor/Source/MeshEditor/MeshEditorMode.cpp
+++ b/Engine/Plugins/Editor/MeshEditor/Source/MeshEditor/MeshEditorMode.cpp
@@ -2328,23 +2328,23 @@ bool FMeshEditorMode::AssignMaterialToSelectedPolygons( UMaterialInterface* Sele
 		{
 			UEditableMesh* EditableMesh = MeshAndPolygons.Key;
 
-			int32 MaterialIndex = INDEX_NONE;
-			FName ImportedMaterialName = NAME_None;
-			//We should be able to add a material at least one adapter must succeed
-			check(EditableMesh->FindOrAddMaterial(SelectedMaterial, MaterialIndex, ImportedMaterialName));
+			EditableMesh->StartModification( EMeshModificationType::Final, EMeshTopologyChange::TopologyChange );
+			{
+				int32 MaterialIndex = INDEX_NONE;
+				FName ImportedMaterialName = NAME_None;
+				//We should be able to add a material at least one adapter must succeed
+				check(EditableMesh->FindOrAddMaterial(SelectedMaterial, MaterialIndex, ImportedMaterialName));
 
-			const UMeshDescription* MeshDescription = EditableMesh->GetMeshDescription();
+				const UMeshDescription* MeshDescription = EditableMesh->GetMeshDescription();
 
-			UPrimitiveComponent* Component = nullptr;
-			for( const FMeshElement& PolygonElement : MeshAndPolygons.Value )
-			{
-				Component = PolygonElement.Component.Get();
-				break;
-			}
-			check( Component != nullptr );
+				UPrimitiveComponent* Component = nullptr;
+				for( const FMeshElement& PolygonElement : MeshAndPolygons.Value )
+				{
+					Component = PolygonElement.Component.Get();
+					break;
+				}
+				check( Component != nullptr );
 
-			EditableMesh->StartModification( EMeshModificationType::Final, EMeshTopologyChange::TopologyChange );
-			{
 				// See if there's a polygon group using this material, and if not create one.
 				// @todo mesheditor: This currently imposes the limitation that each polygon group has a unique material.
 				// Eventually we will need to be able to specify PolygonGroup properties in the editor, and ask the user for
@@ -2368,7 +2368,6 @@ bool FMeshEditorMode::AssignMaterialToSelectedPolygons( UMaterialInterface* Sele
 				// If we didn't find the material being used anywhere, create a new polygon group
 				if( PolygonGroupToAssign == FPolygonGroupID::Invalid )
 				{
-
 					static TArray<FPolygonGroupToCreate> PolygonGroupsToCreate;
 					PolygonGroupsToCreate.Reset( 1 );
 					PolygonGroupsToCreate.Emplace();
diff --git a/Engine/Plugins/Runtime/EditableMesh/Source/EditableMesh/EditableMesh.cpp b/Engine/Plugins/Runtime/EditableMesh/Source/EditableMesh/EditableMesh.cpp
--- a/Engine/Plugins/Runtime/EditableMesh/Source/EditableMesh/EditableMesh.cpp
+++ b/Engine/Plugins/Runtime/EditableMesh/Source/EditableMesh/EditableMesh.cpp
@@ -712,7 +712,7 @@ void UEditableMesh::EndModification( const bool bFromUndo )
 			!bFromUndo &&
 			bAllowCompact )
 		{
-			if( ++PendingCompactCounter == CompactFrequency )
+			if( false )//++PendingCompactCounter == CompactFrequency )
 			{
 				PendingCompactCounter = 0;
 				Compact();
diff --git a/Engine/Plugins/Runtime/EditableMesh/Source/EditableMesh/EditableStaticMeshAdapter.cpp b/Engine/Plugins/Runtime/EditableMesh/Source/EditableMesh/EditableStaticMeshAdapter.cpp
--- a/Engine/Plugins/Runtime/EditableMesh/Source/EditableMesh/EditableStaticMeshAdapter.cpp
+++ b/Engine/Plugins/Runtime/EditableMesh/Source/EditableMesh/EditableStaticMeshAdapter.cpp
@@ -839,18 +839,9 @@ void UEditableStaticMeshAdapter::OnStartModification( const UEditableMesh* Edita
 {
 	// @todo mesheditor undo: We're not using traditional transactions to undo mesh changes yet, but we still want to dirty the mesh package
 	// Also, should we even need the Initializing type? Should we not wait for the first modification before dirtying the package?
-	if( 0 )
-	{
-		this->SetFlags( RF_Transactional );
-		this->Modify();
 
-		StaticMesh->SetFlags( RF_Transactional );
-		StaticMesh->Modify();
-	}
-	else
-	{
-		StaticMesh->MarkPackageDirty();
-	}
+	StaticMesh->SetFlags( RF_Transactional );
+	StaticMesh->Modify();
 }
 
 
