Changed the behavior of the Insert Vertex mesh editor action. Previously it was "insert then move vertex", but this was faulty due to inertia behavior causing the selected point to drift when the element selection changed mid-drag. This could be fixed by restarting the drag action after the call to ApplyDuringDrag(), but I think it works better as an initially constrained placement, followed by the vertex being selected.

diff --git a/Engine/Plugins/Editor/MeshEditor/Source/PolygonModeling/SplitEdge.cpp b/Engine/Plugins/Editor/MeshEditor/Source/PolygonModeling/SplitEdge.cpp
--- a/Engine/Plugins/Editor/MeshEditor/Source/PolygonModeling/SplitEdge.cpp
+++ b/Engine/Plugins/Editor/MeshEditor/Source/PolygonModeling/SplitEdge.cpp
@@ -19,43 +19,13 @@ void USplitEdgeCommand::RegisterUICommand( FBindingContext* BindingContext )
 }
 
 
-bool USplitEdgeCommand::TryStartingToDrag( IMeshEditorModeEditingContract& MeshEditorMode, UViewportInteractor* ViewportInteractor )
+void USplitEdgeCommand::ApplyDuringDrag( IMeshEditorModeEditingContract& MeshEditorMode, UViewportInteractor* ViewportInteractor )
 {
-	bool bHaveEdge = false;
-
 	// Figure out what to split
+	static TMap< class UEditableMesh*, TArray< FMeshElement > > SplitEdgeMeshesAndEdgesToSplit;
 	MeshEditorMode.GetSelectedMeshesAndEdges( /* Out */ SplitEdgeMeshesAndEdgesToSplit );
 
-	SplitEdgeSplitList.Reset();
-
 	if( SplitEdgeMeshesAndEdgesToSplit.Num() > 0 )
-	{
-		for( auto& MeshAndEdges : SplitEdgeMeshesAndEdgesToSplit )
-		{
-			UEditableMesh* EditableMesh = MeshAndEdges.Key;
-			const TArray<FMeshElement>& EdgeElements = MeshAndEdges.Value;
-
-			// Figure out where to split
-			FEdgeID ClosestEdgeID = FEdgeID::Invalid;
-			float Split = 0.0f;
-			const bool bFoundSplit = MeshEditorMode.FindEdgeSplitUnderInteractor( ViewportInteractor, EditableMesh, EdgeElements, /* Out */ ClosestEdgeID, /* Out */ Split );
-
-			if( bFoundSplit )
-			{
-				SplitEdgeSplitList.Add( Split );
-				bHaveEdge = true;
-				break;
-			}
-		}
-	}
-
-	return bHaveEdge;
-}
-
-
-void USplitEdgeCommand::ApplyDuringDrag( IMeshEditorModeEditingContract& MeshEditorMode, UViewportInteractor* ViewportInteractor )
-{
-	if( SplitEdgeMeshesAndEdgesToSplit.Num() > 0 && SplitEdgeSplitList.Num() > 0 )
 	{
 		MeshEditorMode.DeselectAllMeshElements();
 
@@ -69,29 +39,41 @@ void USplitEdgeCommand::ApplyDuringDrag( IMeshEditorModeEditingContract& MeshEdi
 
 			const TArray<FMeshElement>& EdgeElements = MeshAndEdges.Value;
 
-			for( const FMeshElement& EdgeElement : EdgeElements )
+			// Figure out where to split
+			FEdgeID ClosestEdgeID = FEdgeID::Invalid;
+			float Split = 0.0f;
+			const bool bFoundSplit = MeshEditorMode.FindEdgeSplitUnderInteractor( ViewportInteractor, EditableMesh, EdgeElements, /* Out */ ClosestEdgeID, /* Out */ Split );
+
+			if( bFoundSplit )
 			{
-				const FEdgeID EdgeID( EdgeElement.ElementAddress.ElementID );
+				for( const FMeshElement& EdgeElement : EdgeElements )
+				{
+					const FEdgeID EdgeID( EdgeElement.ElementAddress.ElementID );
 
-				static TArray<FVertexID> NewVertexIDs;
-				NewVertexIDs.Reset();
+					static TArray<FVertexID> NewVertexIDs;
+					NewVertexIDs.Reset();
 
-				EditableMesh->SplitEdge( EdgeID, SplitEdgeSplitList, /* Out */ NewVertexIDs );
+					static TArray<float> Splits;	// @todo mesheditor edgeloop: Add support for inserting multiple splits at once!
+					Splits.Reset();
+					Splits.Add( Split );
 
-				// Select all of the new vertices that were created by splitting the edge
-				{
-					for( const FVertexID NewVertexID : NewVertexIDs )
+					EditableMesh->SplitEdge( EdgeID, Splits, /* Out */ NewVertexIDs );
+
+					// Select all of the new vertices that were created by splitting the edge
 					{
-						FMeshElement MeshElementToSelect;
+						for( const FVertexID NewVertexID : NewVertexIDs )
 						{
-							MeshElementToSelect.Component = EdgeElement.Component;
-							MeshElementToSelect.ElementAddress.SubMeshAddress = EdgeElement.ElementAddress.SubMeshAddress;
-							MeshElementToSelect.ElementAddress.ElementType = EEditableMeshElementType::Vertex;
-							MeshElementToSelect.ElementAddress.ElementID = NewVertexID;
+							FMeshElement MeshElementToSelect;
+							{
+								MeshElementToSelect.Component = EdgeElement.Component;
+								MeshElementToSelect.ElementAddress.SubMeshAddress = EdgeElement.ElementAddress.SubMeshAddress;
+								MeshElementToSelect.ElementAddress.ElementType = EEditableMeshElementType::Vertex;
+								MeshElementToSelect.ElementAddress.ElementID = NewVertexID;
+							}
+
+							// Queue selection of this new element.  We don't want it to be part of the current action.
+							MeshElementsToSelect.Add( MeshElementToSelect );
 						}
-
-						// Queue selection of this new element.  We don't want it to be part of the current action.
-						MeshElementsToSelect.Add( MeshElementToSelect );
 					}
 				}
 			}
diff --git a/Engine/Plugins/Editor/MeshEditor/Source/PolygonModeling/SplitEdge.h b/Engine/Plugins/Editor/MeshEditor/Source/PolygonModeling/SplitEdge.h
--- a/Engine/Plugins/Editor/MeshEditor/Source/PolygonModeling/SplitEdge.h
+++ b/Engine/Plugins/Editor/MeshEditor/Source/PolygonModeling/SplitEdge.h
@@ -15,12 +15,10 @@ class USplitEdgeCommand : public UMeshEditorEditCommand
 protected:
 
 	USplitEdgeCommand()
-	  : SplitEdgeMeshesAndEdgesToSplit(),
-	    SplitEdgeSplitList()
 	{
 		UndoText = NSLOCTEXT( "MeshEditor", "UndoSplitEdge", "Split Edge" );
-		bNeedsHoverLocation = false;
-		bNeedsDraggingInitiated = true;
+		bNeedsHoverLocation = true;
+		bNeedsDraggingInitiated = false;
 	}
 
 	// Overrides
@@ -29,15 +27,5 @@ protected:
 		return EEditableMeshElementType::Edge;
 	}
 	virtual void RegisterUICommand( class FBindingContext* BindingContext ) override;
-	virtual bool TryStartingToDrag( IMeshEditorModeEditingContract& MeshEditorMode, class UViewportInteractor* ViewportInteractor ) override;
 	virtual void ApplyDuringDrag( IMeshEditorModeEditingContract& MeshEditorMode, class UViewportInteractor* ViewportInteractor ) override;
-
-protected:
-
-	/** When splitting an edge and dragging a vertex, this is the list of edges that will be split */
-	TMap< class UEditableMesh*, TArray< FMeshElement > > SplitEdgeMeshesAndEdgesToSplit;
-
-	/** When splitting an edge and dragging a vertex, this is the list of split positions along those edges */
-	TArray<float> SplitEdgeSplitList;
-
 };
