Fixed shadowed variable warnings for Linux build.

diff --git a/Engine/Plugins/Editor/MeshEditor/Source/MeshEditor/MeshEditorMode.cpp b/Engine/Plugins/Editor/MeshEditor/Source/MeshEditor/MeshEditorMode.cpp
--- a/Engine/Plugins/Editor/MeshEditor/Source/MeshEditor/MeshEditorMode.cpp
+++ b/Engine/Plugins/Editor/MeshEditor/Source/MeshEditor/MeshEditorMode.cpp
@@ -3113,10 +3113,10 @@ void FMeshEditorMode::UpdateActiveAction( const bool bIsActionFinishing )
 						const FVector DirectionY = FVector::CrossProduct( PlaneNormal, DirectionX );
 
 						// Transforms a point in 3D space into the basis on the plane described by an origin and two orthogonal direction vectors on the plane
-						auto ToPlaneBasis = []( const FVector& Origin, const FVector& DirectionX, const FVector& DirectionY, const FVector& Point )
+						auto ToPlaneBasis = []( const FVector& InOrigin, const FVector& InDirectionX, const FVector& InDirectionY, const FVector& InPoint )
 						{
-							const FVector Offset( Point - Origin );
-							return FVector2D( FVector::DotProduct( Offset, DirectionX ), FVector::DotProduct( Offset, DirectionY ) );
+							const FVector Offset( InPoint - InOrigin );
+							return FVector2D( FVector::DotProduct( Offset, InDirectionX ), FVector::DotProduct( Offset, InDirectionY ) );
 						};
 
 						// Determine whether two line segments intersect in 2D space
@@ -4610,7 +4610,7 @@ bool FMeshEditorMode::FrustumSelect( const FConvexVolume& InFrustum, FEditorView
 		MarqueeSelectTransaction = MakeUnique<FScopedTransaction>( LOCTEXT( "MarqueeSelectElements", "Marquee Select Elements" ) );
 		bMarqueeSelectTransactionActive = true;
 
-		auto OnMenuDismissed = [ this ]( TSharedRef<IMenu> Menu )
+		auto OnMenuDismissed = [ this ]( TSharedRef<IMenu> InMenu )
 		{
 			// End transaction here.
 			// This will actually be released in the next Tick() - this is necessary because the OnMenuDismissed callback happens *before* the action has been executed,
