MeshEditor mode now defaults to disabled, while it is still in active development and subject to change. Two console commands (MeshEditor.Enable and MeshEditor.Disable) are provided in order to enable/disable it.

diff --git a/Engine/Plugins/Editor/MeshEditor/Source/MeshEditor/MeshEditor.Build.cs b/Engine/Plugins/Editor/MeshEditor/Source/MeshEditor/MeshEditor.Build.cs
--- a/Engine/Plugins/Editor/MeshEditor/Source/MeshEditor/MeshEditor.Build.cs
+++ b/Engine/Plugins/Editor/MeshEditor/Source/MeshEditor/MeshEditor.Build.cs
@@ -23,7 +23,8 @@ namespace UnrealBuildTool.Rules
                     "VREditor",
 					"Projects",
                     "ShaderCore",
-                    "RHI"
+                    "RHI",
+                    "LevelEditor"
                 }
             );
 
@@ -37,8 +38,7 @@ namespace UnrealBuildTool.Rules
 
 			DynamicallyLoadedModuleNames.AddRange(
 				new string[] {
-                    "ContentBrowser",
-					"LevelEditor",
+                    "ContentBrowser"
 				}
 			);
         }
diff --git a/Engine/Plugins/Editor/MeshEditor/Source/MeshEditor/MeshEditorModule.cpp b/Engine/Plugins/Editor/MeshEditor/Source/MeshEditor/MeshEditorModule.cpp
--- a/Engine/Plugins/Editor/MeshEditor/Source/MeshEditor/MeshEditorModule.cpp
+++ b/Engine/Plugins/Editor/MeshEditor/Source/MeshEditor/MeshEditorModule.cpp
@@ -14,6 +14,8 @@
 #include "VREditorMode.h"
 #include "Editor.h"
 #include "IMeshEditorModeUIContract.h"
+#include "LevelEditor.h"
+#include "LevelEditorModesActions.h"
 
 
 #define LOCTEXT_NAMESPACE "MeshEditor"
@@ -21,7 +23,10 @@
 class FMeshEditorModule : public IModuleInterface
 {
 public:
-	FMeshEditorModule()
+	FMeshEditorModule() :
+		bIsEnabled( false ),
+		MeshEditorEnable( TEXT( "MeshEditor.Enable" ), TEXT( "Makes MeshEditor mode available" ), FConsoleCommandDelegate::CreateRaw( this, &FMeshEditorModule::Register ) ),
+		MeshEditorDisable( TEXT( "MeshEditor.Disable" ), TEXT( "Makes MeshEditor mode unavailable" ), FConsoleCommandDelegate::CreateRaw( this, &FMeshEditorModule::Unregister ) )
 	{
 	}
 
@@ -42,6 +47,9 @@ protected:
 		return MeshEditorFeatureName;
 	}
 
+	void Register();
+	void Unregister();
+
 	/** Adds items to the VR Radial menu for mesh editing mode */
 	void FillVRRadialMenuModes( class FMenuBuilder& MenuBuilder );
 
@@ -54,17 +62,41 @@ protected:
 	/** Should the mesh edit button be enabled */
 	bool IsMeshEditModeButtonEnabled( EEditableMeshElementType InMode );
 
-
 	/** Menu extension for the VR Editor's 'Modes' menu */
 	TSharedPtr<const class FExtensionBase> VRRadialMenuModesExtension;
 
+	/** Console commands for enabling/disabling mesh editor mode while it is still in development */
+	FAutoConsoleCommand MeshEditorEnable;
+	FAutoConsoleCommand MeshEditorDisable;
+
+	/** Whether mesh editor mode is enabled: currently defaults to false */
+	bool bIsEnabled;
 };
 
 
 void FMeshEditorModule::StartupModule()
 {
+	// Small hack while we're controlling whether mesh editor mode should be enabled on startup or not.
+	if( bIsEnabled )
+	{
+		bIsEnabled = false;
+		Register();
+	}
+}
+
+
+void FMeshEditorModule::Register()
+{
+	if( bIsEnabled )
+	{
+		return;
+	}
+
+	bIsEnabled = true;
+
 	FMeshEditorStyle::Initialize();
 
+	FLevelEditorModesCommands::Unregister();
 	FEditorModeRegistry::Get().RegisterMode<FMeshEditorMode>(
 		GetEditorModeID(),
 		LOCTEXT( "ModeName", "Mesh Editor" ),
@@ -97,6 +129,19 @@ void FMeshEditorModule::StartupModule()
 
 void FMeshEditorModule::ShutdownModule()
 {
+	Unregister();
+}
+
+
+void FMeshEditorModule::Unregister()
+{
+	if( !bIsEnabled )
+	{
+		return;
+	}
+
+	bIsEnabled = false;
+
 	{
 		if( IVREditorModule::IsAvailable() )
 		{
@@ -112,6 +157,7 @@ void FMeshEditorModule::ShutdownModule()
 		SettingsModule->UnregisterSettings( "Editor", "ContentEditors", "MeshEditor" );
 	}
 
+	FLevelEditorModesCommands::Unregister();
 	FEditorModeRegistry::Get().UnregisterMode(GetEditorModeID());
 
 	FMeshEditorStyle::Shutdown();
