Copied from Dev-Editor CL 4128420

Fixed various Mesh Editor issues with overlay meshes.

diff --git a/Engine/Plugins/Editor/MeshEditor/Source/MeshEditor/MeshEditorMode.cpp b/Engine/Plugins/Editor/MeshEditor/Source/MeshEditor/MeshEditorMode.cpp
--- a/Engine/Plugins/Editor/MeshEditor/Source/MeshEditor/MeshEditorMode.cpp
+++ b/Engine/Plugins/Editor/MeshEditor/Source/MeshEditor/MeshEditorMode.cpp
@@ -843,6 +843,15 @@ void FMeshEditorMode::Enter()
 	// Call parent implementation
 	FEdMode::Enter();
 
+	// Initialize selection sets and caches
+	SelectedComponentsAndEditableMeshes.Reset();
+	SelectedEditableMeshes.Reset();
+	SelectedMeshElements.Empty();
+	SelectedVertices.Empty();
+	SelectedEdges.Empty();
+	SelectedPolygons.Empty();
+	ComponentToWireframeComponentMap.Empty();
+
 	// Notify when the map changes
 	FLevelEditorModule& LevelEditorModule = FModuleManager::LoadModuleChecked<FLevelEditorModule>( "LevelEditor" );
 	LevelEditorModule.OnMapChanged().AddRaw( this, &FMeshEditorMode::OnMapChanged );
@@ -4736,6 +4745,7 @@ void FMeshEditorMode::UpdateSelectedEditableMeshes()
 
 				// Don't bother with editor-only 'helper' actors, we never want to visualize or edit geometry on those
 				if( !Component->IsEditorOnly() &&
+					Component->GetCollisionEnabled() != ECollisionEnabled::NoCollision &&
 					( Component->GetOwner() == nullptr || !Component->GetOwner()->IsEditorOnly() ) )
 				{
 					const int32 LODIndex = 0;			// @todo mesheditor: We'll want to select an LOD to edit in various different wants (LOD that's visible, or manual user select, etc.)
