Fix build

diff --git a/Engine/Plugins/Editor/MeshEditor/Source/MeshEditor/MeshEditorUtilities.cpp b/Engine/Plugins/Editor/MeshEditor/Source/MeshEditor/MeshEditorUtilities.cpp
--- a/Engine/Plugins/Editor/MeshEditor/Source/MeshEditor/MeshEditorUtilities.cpp
+++ b/Engine/Plugins/Editor/MeshEditor/Source/MeshEditor/MeshEditorUtilities.cpp
@@ -11,7 +11,7 @@ bool FMeshEditorUtilities::AssignMaterialToPolygons( UMaterialInterface* Selecte
 {
 	EditableMesh->StartModification( EMeshModificationType::Final, EMeshTopologyChange::TopologyChange );
 	{
-		const UMeshDescription* MeshDescription = EditableMesh->GetMeshDescription();
+		const FMeshDescription* MeshDescription = EditableMesh->GetMeshDescription();
 
 		UPrimitiveComponent* Component = nullptr;
 		for( const FMeshElement& PolygonElement : PolygonElements )
