Fixed crash when entering mesh editor mode after having loaded a new level.

diff --git a/Engine/Plugins/Editor/MeshEditor/Source/MeshEditor/MeshEditorMode.cpp b/Engine/Plugins/Editor/MeshEditor/Source/MeshEditor/MeshEditorMode.cpp
--- a/Engine/Plugins/Editor/MeshEditor/Source/MeshEditor/MeshEditorMode.cpp
+++ b/Engine/Plugins/Editor/MeshEditor/Source/MeshEditor/MeshEditorMode.cpp
@@ -414,6 +414,11 @@ FMeshEditorMode::FMeshEditorMode()
 	  PreviewRevertChanges(),
 	  ActiveActionModifiedMeshes(),
 	  MeshEditorModeProxyObject( nullptr ),
+	  WireframeComponentContainer( nullptr ),
+	  HoveredElementsComponent( nullptr ),
+	  SelectedElementsComponent( nullptr ),
+	  SelectedSubDElementsComponent( nullptr ),
+	  DebugNormalsComponent( nullptr ),
 	  ActiveActionInteractor( nullptr ),
 	  bActiveActionNeedsHoverLocation( false ),
 	  bIsFirstActiveActionUpdate( false ),
@@ -462,14 +467,6 @@ FMeshEditorMode::FMeshEditorMode()
 }
 
 
-void FMeshEditorMode::Initialize()
-{
-	FActorSpawnParameters ActorSpawnParameters;
-	ActorSpawnParameters.ObjectFlags |= RF_Transient;
-	WireframeComponentContainer = GetWorld()->SpawnActor<AActor>( ActorSpawnParameters );
-}
-
-
 FMeshEditorMode::~FMeshEditorMode()
 {
 	// Unregister mesh editor actions
@@ -486,9 +483,6 @@ FMeshEditorMode::~FMeshEditorMode()
 		CachedEditableMesh.Value.EditableMesh->OnElementIDsRemapped().RemoveAll( this );
 	}
 
-	WireframeComponentContainer->Destroy();
-	WireframeComponentContainer = nullptr;
-
 	MeshEditorModeProxyObject = nullptr;
 	AssetContainer = nullptr;
 }
@@ -859,8 +853,13 @@ void FMeshEditorMode::Enter()
 
 	FEditorDelegates::EndPIE.AddRaw( this, &FMeshEditorMode::OnEndPIE );
 
+	// Create wireframe component container
+	FActorSpawnParameters ActorSpawnParameters;
+	ActorSpawnParameters.ObjectFlags |= RF_Transient;
+	WireframeComponentContainer = GetWorld()->SpawnActor<AActor>( ActorSpawnParameters );
+
 	// Add overlay component for rendering hovered elements
-	HoveredElementsComponent= NewObject<UOverlayComponent>( WireframeComponentContainer );
+	HoveredElementsComponent = NewObject<UOverlayComponent>( WireframeComponentContainer );
 	HoveredElementsComponent->SetLineMaterial( OverlayLineMaterial );
 	HoveredElementsComponent->SetPointMaterial( OverlayPointMaterial );
 	HoveredElementsComponent->TranslucencySortPriority = 400;
@@ -1011,6 +1010,13 @@ void FMeshEditorMode::Exit()
 	SelectedSubDElementsComponent->DestroyComponent();
 	SelectedElementsComponent->DestroyComponent();
 	HoveredElementsComponent->DestroyComponent();
+	DebugNormalsComponent = nullptr;
+	SelectedSubDElementsComponent = nullptr;
+	SelectedElementsComponent = nullptr;
+	HoveredElementsComponent = nullptr;
+
+	WireframeComponentContainer->Destroy();
+	WireframeComponentContainer = nullptr;
 
 	FEditorDelegates::EndPIE.RemoveAll( this );
 
diff --git a/Engine/Plugins/Editor/MeshEditor/Source/MeshEditor/MeshEditorMode.h b/Engine/Plugins/Editor/MeshEditor/Source/MeshEditor/MeshEditorMode.h
--- a/Engine/Plugins/Editor/MeshEditor/Source/MeshEditor/MeshEditorMode.h
+++ b/Engine/Plugins/Editor/MeshEditor/Source/MeshEditor/MeshEditorMode.h
@@ -117,7 +117,6 @@ public:
 protected:
 
 	// FEdMode interface
-	virtual void Initialize() override;
 	virtual void Enter() override;
 	virtual void Exit() override;
 	virtual void Tick( FEditorViewportClient* ViewportClient, float DeltaTime ) override;
