Fixed plugins build

#rnx

diff --git a/Engine/Plugins/Editor/MeshEditor/Source/MeshEditor/OverlayComponent.cpp b/Engine/Plugins/Editor/MeshEditor/Source/MeshEditor/OverlayComponent.cpp
--- a/Engine/Plugins/Editor/MeshEditor/Source/MeshEditor/OverlayComponent.cpp
+++ b/Engine/Plugins/Editor/MeshEditor/Source/MeshEditor/OverlayComponent.cpp
@@ -253,7 +253,7 @@ public:
 					Mesh.MaterialRenderProxy = MeshBatchData.MaterialProxy;
 
 					FDynamicPrimitiveUniformBuffer& DynamicPrimitiveUniformBuffer = Collector.AllocateOneFrameResource<FDynamicPrimitiveUniformBuffer>();
-					DynamicPrimitiveUniformBuffer.Set(GetLocalToWorld(), GetLocalToWorld(), GetBounds(), GetLocalBounds(), true, false, UseEditorDepthTest());
+					DynamicPrimitiveUniformBuffer.Set(GetLocalToWorld(), GetLocalToWorld(), GetBounds(), GetLocalBounds(), true, false, UseEditorDepthTest(), false);
 					BatchElement.PrimitiveUniformBufferResource = &DynamicPrimitiveUniformBuffer.UniformBuffer;
 
 					BatchElement.FirstIndex = MeshBatchData.StartIndex;
diff --git a/Engine/Plugins/Editor/MeshEditor/Source/MeshEditor/WireframeMeshComponent.cpp b/Engine/Plugins/Editor/MeshEditor/Source/MeshEditor/WireframeMeshComponent.cpp
--- a/Engine/Plugins/Editor/MeshEditor/Source/MeshEditor/WireframeMeshComponent.cpp
+++ b/Engine/Plugins/Editor/MeshEditor/Source/MeshEditor/WireframeMeshComponent.cpp
@@ -344,7 +344,7 @@ public:
 				Mesh.MaterialRenderProxy = MaterialProxy;
 
 				FDynamicPrimitiveUniformBuffer& DynamicPrimitiveUniformBuffer = Collector.AllocateOneFrameResource<FDynamicPrimitiveUniformBuffer>();
-				DynamicPrimitiveUniformBuffer.Set(GetLocalToWorld(), GetLocalToWorld(), GetBounds(), GetLocalBounds(), true, false, UseEditorDepthTest());
+				DynamicPrimitiveUniformBuffer.Set(GetLocalToWorld(), GetLocalToWorld(), GetBounds(), GetLocalBounds(), true, false, UseEditorDepthTest(), false);
 				BatchElement.PrimitiveUniformBufferResource = &DynamicPrimitiveUniformBuffer.UniformBuffer;
 
 				BatchElement.FirstIndex = 0;
