Fracture: Remove the button from the Mesh Editing Mode to access the depricated Fracture sub-mode

diff --git a/Engine/Plugins/Editor/MeshEditor/Source/MeshEditor/MeshEditorModeToolkit.cpp b/Engine/Plugins/Editor/MeshEditor/Source/MeshEditor/MeshEditorModeToolkit.cpp
--- a/Engine/Plugins/Editor/MeshEditor/Source/MeshEditor/MeshEditorModeToolkit.cpp
+++ b/Engine/Plugins/Editor/MeshEditor/Source/MeshEditor/MeshEditorModeToolkit.cpp
@@ -304,12 +304,14 @@ void SMeshEditorModeControls::Construct( const FArguments& InArgs, IMeshEditorMo
 				[
 					SNew( SMeshEditorSelectionModeWidget, MeshEditorMode, EEditableMeshElementType::Vertex, LOCTEXT( "Vertex", "Vertex" ) )
 				]
+
+				/* // Remove access to the Fracture Sub-Mode
 				+ SHorizontalBox::Slot()
 				.FillWidth(1)
 				.Padding(2)
 				[
 					SNew(SMeshEditorSelectionModeWidget, MeshEditorMode, EEditableMeshElementType::Fracture, LOCTEXT("Fracture", "Fracture"))
-				]
+				]*/
 
 			]
 			+SVerticalBox::Slot()
