Fracture: Removes GeometryCollectionFactory and GeometryCollectionCacheFactory to prohibit creation of Geometry Collections outside of the Fracture Panel.

diff --git a/Engine/Plugins/Editor/MeshEditor/Source/PolygonModeling/ClusterMeshCommand.cpp b/Engine/Plugins/Editor/MeshEditor/Source/PolygonModeling/ClusterMeshCommand.cpp
deleted file mode 100644
diff --git a/Engine/Plugins/Editor/MeshEditor/Source/PolygonModeling/ClusterMeshCommand.h b/Engine/Plugins/Editor/MeshEditor/Source/PolygonModeling/ClusterMeshCommand.h
deleted file mode 100644
diff --git a/Engine/Plugins/Editor/MeshEditor/Source/PolygonModeling/ConvertMeshToGeometryCollectionCommand.cpp b/Engine/Plugins/Editor/MeshEditor/Source/PolygonModeling/ConvertMeshToGeometryCollectionCommand.cpp
deleted file mode 100644
diff --git a/Engine/Plugins/Editor/MeshEditor/Source/PolygonModeling/ConvertMeshToGeometryCollectionCommand.h b/Engine/Plugins/Editor/MeshEditor/Source/PolygonModeling/ConvertMeshToGeometryCollectionCommand.h
deleted file mode 100644
diff --git a/Engine/Plugins/Editor/MeshEditor/Source/PolygonModeling/FractureMeshCommand.cpp b/Engine/Plugins/Editor/MeshEditor/Source/PolygonModeling/FractureMeshCommand.cpp
deleted file mode 100644
diff --git a/Engine/Plugins/Editor/MeshEditor/Source/PolygonModeling/FractureMeshCommand.h b/Engine/Plugins/Editor/MeshEditor/Source/PolygonModeling/FractureMeshCommand.h
deleted file mode 100644
diff --git a/Engine/Plugins/Editor/MeshEditor/Source/PolygonModeling/GeometryCollectionCommandCommon.cpp b/Engine/Plugins/Editor/MeshEditor/Source/PolygonModeling/GeometryCollectionCommandCommon.cpp
--- a/Engine/Plugins/Editor/MeshEditor/Source/PolygonModeling/GeometryCollectionCommandCommon.cpp
+++ b/Engine/Plugins/Editor/MeshEditor/Source/PolygonModeling/GeometryCollectionCommandCommon.cpp
@@ -13,11 +13,11 @@
 #include "SceneOutlinerDelegates.h"
 #include "MeshFractureSettings.h"
 #include "AssetRegistryModule.h"
-#include "GeometryCollection/GeometryCollectionFactory.h"
+// #include "GeometryCollection/GeometryCollectionFactory.h"
 #include "GeometryCollection/GeometryCollectionConversion.h"
 #include "GeometryCollection/GeometryCollectionClusteringUtility.h"
 #include "GeometryCollection/GeometryCollectionAlgo.h"
-#include "GeometryCollection/GeometryCollectionFactory.h"
+// #include "GeometryCollection/GeometryCollectionFactory.h"
 #include "AssetToolsModule.h"
 #include "FractureToolDelegates.h"
 
@@ -98,6 +98,7 @@ namespace CommandCommon
 // Note that this isn't really creating an actor representing the source mesh, but is only copying over the materials and transform.  This
 // is used in conjunction with other methods for filling the resulting Actor's GeometryCollection (ie: Clustering operations, fracturing, etc)
 // #todo(dmp): at some point we should consider refactoring or renaming this.
+/*
 AGeometryCollectionActor* FGeometryCollectionCommandCommon::CreateNewGeometryActor(const FString& Name, const FTransform& Transform, UEditableMesh* SourceMesh, bool AddMaterials)
 {
 	// create an asset package first
@@ -155,7 +156,7 @@ AGeometryCollectionActor* FGeometryCollectionCommandCommon::CreateNewGeometryAct
 	Package->SetDirtyFlag(true);
 
 	return NewActor;
-}
+}*/
 
 void FGeometryCollectionCommandCommon::RemoveActor(AActor* Actor)
 {
@@ -301,6 +302,7 @@ UEditableMesh* FGeometryCollectionCommandCommon::GetEditableMeshForComponent(UAc
 	return nullptr;
 }
 
+/*
 UPackage* FGeometryCollectionCommandCommon::CreateGeometryCollectionPackage(UGeometryCollection*& GeometryCollection)
 {
 	UPackage* Package = CreatePackage(NULL, TEXT("/Game/GeometryCollectionAsset"));
@@ -309,6 +311,7 @@ UPackage* FGeometryCollectionCommandCommon::CreateGeometryCollectionPackage(UGeo
 			FName("GeometryCollectionAsset"), RF_Standalone | RF_Public, NULL, GWarn));		
 	return Package;
 }
+*/
 
 void FGeometryCollectionCommandCommon::AddSingleRootNodeIfRequired(UGeometryCollection* GeometryCollectionObject)
 {
diff --git a/Engine/Plugins/Editor/MeshEditor/Source/PolygonModeling/GeometryCollectionCommandCommon.h b/Engine/Plugins/Editor/MeshEditor/Source/PolygonModeling/GeometryCollectionCommandCommon.h
--- a/Engine/Plugins/Editor/MeshEditor/Source/PolygonModeling/GeometryCollectionCommandCommon.h
+++ b/Engine/Plugins/Editor/MeshEditor/Source/PolygonModeling/GeometryCollectionCommandCommon.h
@@ -32,10 +32,10 @@ protected:
 	UEditableMesh* GetEditableMeshForComponent(UActorComponent* PrimitiveActorComponent, TArray<UEditableMesh *>& SelectedMeshes);
 
 	/** Create a Geometry Actor */
-	class AGeometryCollectionActor* CreateNewGeometryActor(const FString& Name, const FTransform& Transform, UEditableMesh* SourceMesh, bool AddMaterials = false);
+	// class AGeometryCollectionActor* CreateNewGeometryActor(const FString& Name, const FTransform& Transform, UEditableMesh* SourceMesh, bool AddMaterials = false);
 
 	/** Create a geomerty Collection asset */
-	class UPackage* CreateGeometryCollectionPackage(UGeometryCollection*& GeometryCollection);
+	// class UPackage* CreateGeometryCollectionPackage(UGeometryCollection*& GeometryCollection);
 
 	/** Where fracturing is concerned we expect a single parent root node */
 	void AddSingleRootNodeIfRequired(UGeometryCollection* GeometryCollection);
diff --git a/Engine/Plugins/Editor/MeshEditor/Source/PolygonModeling/MergeToClusterCommand.cpp b/Engine/Plugins/Editor/MeshEditor/Source/PolygonModeling/MergeToClusterCommand.cpp
deleted file mode 100644
diff --git a/Engine/Plugins/Editor/MeshEditor/Source/PolygonModeling/MergeToClusterCommand.h b/Engine/Plugins/Editor/MeshEditor/Source/PolygonModeling/MergeToClusterCommand.h
deleted file mode 100644
diff --git a/Engine/Plugins/Experimental/GeometryCollectionPlugin/Source/GeometryCollectionEditor/Private/GeometryCollection/ActorFactoryGeometryCollection.cpp b/Engine/Plugins/Experimental/GeometryCollectionPlugin/Source/GeometryCollectionEditor/Private/GeometryCollection/ActorFactoryGeometryCollection.cpp
--- a/Engine/Plugins/Experimental/GeometryCollectionPlugin/Source/GeometryCollectionEditor/Private/GeometryCollection/ActorFactoryGeometryCollection.cpp
+++ b/Engine/Plugins/Experimental/GeometryCollectionPlugin/Source/GeometryCollectionEditor/Private/GeometryCollection/ActorFactoryGeometryCollection.cpp
@@ -4,7 +4,6 @@
 #include "GeometryCollection/GeometryCollectionActor.h"
 #include "GeometryCollection/GeometryCollection.h"
 #include "GeometryCollection/GeometryCollectionComponent.h"
-#include "GeometryCollection/GeometryCollectionFactory.h"
 #include "Misc/FileHelper.h"
 #include "ActorFactories/ActorFactory.h"
 #include "AssetData.h"
diff --git a/Engine/Plugins/Experimental/GeometryCollectionPlugin/Source/GeometryCollectionEditor/Private/GeometryCollection/GeometryCollectionCacheFactory.cpp b/Engine/Plugins/Experimental/GeometryCollectionPlugin/Source/GeometryCollectionEditor/Private/GeometryCollection/GeometryCollectionCacheFactory.cpp
deleted file mode 100644
diff --git a/Engine/Plugins/Experimental/GeometryCollectionPlugin/Source/GeometryCollectionEditor/Private/GeometryCollection/GeometryCollectionConversion.cpp b/Engine/Plugins/Experimental/GeometryCollectionPlugin/Source/GeometryCollectionEditor/Private/GeometryCollection/GeometryCollectionConversion.cpp
--- a/Engine/Plugins/Experimental/GeometryCollectionPlugin/Source/GeometryCollectionEditor/Private/GeometryCollection/GeometryCollectionConversion.cpp
+++ b/Engine/Plugins/Experimental/GeometryCollectionPlugin/Source/GeometryCollectionEditor/Private/GeometryCollection/GeometryCollectionConversion.cpp
@@ -15,7 +15,7 @@
 #include "GeometryCollection/GeometryCollection.h"
 #include "GeometryCollection/GeometryCollectionActor.h"
 #include "GeometryCollection/GeometryCollectionAlgo.h"
-#include "GeometryCollection/GeometryCollectionFactory.h"
+// #include "GeometryCollection/GeometryCollectionFactory.h"
 #include "GeometryCollection/GeometryCollectionComponent.h"
 #include "GeometryCollection/GeometryCollectionClusteringUtility.h"
 #include "GeometryCollection/GeometryCollectionUtility.h"
@@ -426,15 +426,3 @@ void FGeometryCollectionConversion::AppendSkeletalMesh(const USkeletalMesh* Skel
 	}
 }
 
-void FGeometryCollectionConversion::CreateGeometryCollectionCommand(UWorld * World)
-{
-	UPackage* Package = CreatePackage(NULL, TEXT("/Game/GeometryCollectionAsset"));
-	auto GeometryCollectionFactory = NewObject<UGeometryCollectionFactory>();
-	UGeometryCollection* GeometryCollection = static_cast<UGeometryCollection*>(
-		GeometryCollectionFactory->FactoryCreateNew(UGeometryCollection::StaticClass(), Package,
-			FName("GeometryCollectionAsset"), RF_Standalone | RF_Public, NULL, GWarn));
-	FAssetRegistryModule::AssetCreated(GeometryCollection);
-	Package->SetDirtyFlag(true);
-}
-
-
diff --git a/Engine/Plugins/Experimental/GeometryCollectionPlugin/Source/GeometryCollectionEditor/Private/GeometryCollection/GeometryCollectionFactory.cpp b/Engine/Plugins/Experimental/GeometryCollectionPlugin/Source/GeometryCollectionEditor/Private/GeometryCollection/GeometryCollectionFactory.cpp
deleted file mode 100644
diff --git a/Engine/Plugins/Experimental/GeometryCollectionPlugin/Source/GeometryCollectionEditor/Public/GeometryCollection/GeometryCollectionCacheFactory.h b/Engine/Plugins/Experimental/GeometryCollectionPlugin/Source/GeometryCollectionEditor/Public/GeometryCollection/GeometryCollectionCacheFactory.h
deleted file mode 100644
diff --git a/Engine/Plugins/Experimental/GeometryCollectionPlugin/Source/GeometryCollectionEditor/Public/GeometryCollection/GeometryCollectionConversion.h b/Engine/Plugins/Experimental/GeometryCollectionPlugin/Source/GeometryCollectionEditor/Public/GeometryCollection/GeometryCollectionConversion.h
--- a/Engine/Plugins/Experimental/GeometryCollectionPlugin/Source/GeometryCollectionEditor/Public/GeometryCollection/GeometryCollectionConversion.h
+++ b/Engine/Plugins/Experimental/GeometryCollectionPlugin/Source/GeometryCollectionEditor/Public/GeometryCollection/GeometryCollectionConversion.h
@@ -33,10 +33,4 @@ public:
 	*/
 	static void AppendSkeletalMesh(const USkeletalMesh* SkeletalMesh, const USkeletalMeshComponent *SkeletalMeshComponent, const FTransform & SkeletalMeshTransform, UGeometryCollection * GeometryCollection, bool ReindexMaterials = true);
 
-	/**
-	*  Command invoked from "GeometryCollection.CreatGeometryCollection", uses the selected Actors to create a GeometryCollection Asset
-	*  @param World
-	*/
-	static void CreateGeometryCollectionCommand(UWorld * World);
-
 };
\ No newline at end of file
diff --git a/Engine/Plugins/Experimental/GeometryCollectionPlugin/Source/GeometryCollectionEditor/Public/GeometryCollection/GeometryCollectionFactory.h b/Engine/Plugins/Experimental/GeometryCollectionPlugin/Source/GeometryCollectionEditor/Public/GeometryCollection/GeometryCollectionFactory.h
deleted file mode 100644
