Disable the MeshEditor mode by default. It was enabled in dev-destruction but its not needed anymore and shouldn't be available by default.

diff --git a/Engine/Plugins/Editor/MeshEditor/Source/MeshEditor/MeshEditorModule.cpp b/Engine/Plugins/Editor/MeshEditor/Source/MeshEditor/MeshEditorModule.cpp
--- a/Engine/Plugins/Editor/MeshEditor/Source/MeshEditor/MeshEditorModule.cpp
+++ b/Engine/Plugins/Editor/MeshEditor/Source/MeshEditor/MeshEditorModule.cpp
@@ -24,7 +24,7 @@ class FMeshEditorModule : public IModuleInterface
 {
 public:
 	FMeshEditorModule() :
-		bIsEnabled( true ),
+		bIsEnabled( false ),
 		MeshEditorEnable( TEXT( "MeshEditor.Enable" ), TEXT( "Makes MeshEditor mode available" ), FConsoleCommandDelegate::CreateRaw( this, &FMeshEditorModule::Register ) ),
 		MeshEditorDisable( TEXT( "MeshEditor.Disable" ), TEXT( "Makes MeshEditor mode unavailable" ), FConsoleCommandDelegate::CreateRaw( this, &FMeshEditorModule::Unregister ) )
 	{
