Fix a couple derecation warnings

diff --git a/Engine/Plugins/Editor/MeshEditor/Source/PolygonModeling/SplitMesh.cpp b/Engine/Plugins/Editor/MeshEditor/Source/PolygonModeling/SplitMesh.cpp
--- a/Engine/Plugins/Editor/MeshEditor/Source/PolygonModeling/SplitMesh.cpp
+++ b/Engine/Plugins/Editor/MeshEditor/Source/PolygonModeling/SplitMesh.cpp
@@ -283,7 +283,7 @@ void USplitMeshCommand::Execute(IMeshEditorModeEditingContract& MeshEditorMode)
 		NewPackageName = UPackageTools::SanitizePackageName(NewPackageName);
 		UPackage* NewPackage = CreatePackage(nullptr, *NewPackageName);
 		UStaticMesh* NewStaticMesh = NewObject<UStaticMesh>(NewPackage, *NewMeshName, RF_Public);
-		FMeshDescription* NewMeshDescription = NewStaticMesh->GetOriginalMeshDescription(0);
+		FMeshDescription* NewMeshDescription = NewStaticMesh->GetMeshDescription(0);
 		check(NewMeshDescription);
 		new (NewStaticMesh->SourceModels) FStaticMeshSourceModel();
 
@@ -363,7 +363,7 @@ void USplitMeshCommand::Execute(IMeshEditorModeEditingContract& MeshEditorMode)
 		{
 			NewStaticMesh->StaticMaterials.Add(UMaterial::GetDefaultMaterial(MD_Surface));
 		}
-		NewStaticMesh->CommitOriginalMeshDescription(0);
+		NewStaticMesh->CommitMeshDescription(0);
 		NewStaticMesh->Build();
 		NewStaticMesh->PostEditChange();
 		AStaticMeshActor* NewMeshActor = Cast<AStaticMeshActor>(AddActor(GetSelectedLevel(), AStaticMeshActor::StaticClass()));
