Temporarily disable USplitMeshCommand::Execute until template compile error specifically on VS 15.8.9 is resolved (marcus.wassmer)

diff --git a/Engine/Plugins/Editor/MeshEditor/Source/PolygonModeling/SplitMesh.cpp b/Engine/Plugins/Editor/MeshEditor/Source/PolygonModeling/SplitMesh.cpp
--- a/Engine/Plugins/Editor/MeshEditor/Source/PolygonModeling/SplitMesh.cpp
+++ b/Engine/Plugins/Editor/MeshEditor/Source/PolygonModeling/SplitMesh.cpp
@@ -125,6 +125,7 @@ void USplitMeshCommand::RegisterUICommand(FBindingContext* BindingContext)
 
 void USplitMeshCommand::Execute(IMeshEditorModeEditingContract& MeshEditorMode)
 {
+#if 0
 	const TArray<UEditableMesh*>& SelectedMeshes = MeshEditorMode.GetSelectedEditableMeshes();
 	const TArray<AActor*> SelectedActors = GetSelectedActors();
 	check(SelectedMeshes.Num() == SelectedActors.Num());
@@ -376,6 +377,7 @@ void USplitMeshCommand::Execute(IMeshEditorModeEditingContract& MeshEditorMode)
 
 		MeshEditorMode.TrackUndo(Mesh, Mesh->MakeUndo());
 	}
+#endif
 }
 
 #undef LOCTEXT_NAMESPACE
