Fixed some template code which wasn't compiling with an old version of VS2017.

diff --git a/Engine/Plugins/Editor/MeshEditor/Source/PolygonModeling/SplitMesh.cpp b/Engine/Plugins/Editor/MeshEditor/Source/PolygonModeling/SplitMesh.cpp
--- a/Engine/Plugins/Editor/MeshEditor/Source/PolygonModeling/SplitMesh.cpp
+++ b/Engine/Plugins/Editor/MeshEditor/Source/PolygonModeling/SplitMesh.cpp
@@ -92,12 +92,11 @@ template <typename ElementIDType>
 static void CopyAllAttributes(TAttributesSet<ElementIDType>& DestAttributesSet, const TAttributesSet<ElementIDType>& SrcAttributesSet, const ElementIDType ElementID)
 {
 	SrcAttributesSet.ForEach(
-		[&DestAttributesSet, ElementID](const FName AttributeName, const auto& AttributeIndicesArray)
+		[&DestAttributesSet, ElementID](const FName AttributeName, auto AttributeArrayRef)
 		{
-			using AttributeType = decltype(AttributeIndicesArray.GetDefaultValue());
-			for(int32 Index = 0; Index < AttributeIndicesArray.GetNumIndices(); ++Index)
+			for(int32 Index = 0; Index < AttributeArrayRef.GetNumIndices(); ++Index)
 			{
-				DestAttributesSet.template SetAttribute<AttributeType>(ElementID, AttributeName, Index, AttributeIndicesArray.Get(ElementID, Index));
+				DestAttributesSet.SetAttribute(ElementID, AttributeName, Index, AttributeArrayRef.Get(ElementID, Index));
 			}
 		}
 	);
@@ -107,12 +106,11 @@ template <typename ElementIDType>
 static void CopyAllAttributesToDifferentElement(TAttributesSet<ElementIDType>& DestAttributesSet, const TAttributesSet<ElementIDType>& SrcAttributesSet, const ElementIDType DstElementID, const ElementIDType SrcElementID)
 {
 	SrcAttributesSet.ForEach(
-		[&DestAttributesSet, DstElementID, SrcElementID](const FName AttributeName, const auto& AttributeIndicesArray)
+		[&DestAttributesSet, DstElementID, SrcElementID](const FName AttributeName, auto AttributeArrayRef)
 		{
-			using AttributeType = decltype(AttributeIndicesArray.GetDefaultValue());
-			for(int32 Index = 0; Index < AttributeIndicesArray.GetNumIndices(); ++Index)
+			for(int32 Index = 0; Index < AttributeArrayRef.GetNumIndices(); ++Index)
 			{
-				DestAttributesSet.template SetAttribute<AttributeType>(DstElementID, AttributeName, Index, AttributeIndicesArray.Get(SrcElementID, Index));
+				DestAttributesSet.SetAttribute(DstElementID, AttributeName, Index, AttributeArrayRef.Get(SrcElementID, Index));
 			}
 		}
 	);
@@ -125,7 +123,6 @@ void USplitMeshCommand::RegisterUICommand(FBindingContext* BindingContext)
 
 void USplitMeshCommand::Execute(IMeshEditorModeEditingContract& MeshEditorMode)
 {
-#if 0
 	const TArray<UEditableMesh*>& SelectedMeshes = MeshEditorMode.GetSelectedEditableMeshes();
 	const TArray<AActor*> SelectedActors = GetSelectedActors();
 	check(SelectedMeshes.Num() == SelectedActors.Num());
@@ -377,7 +374,6 @@ void USplitMeshCommand::Execute(IMeshEditorModeEditingContract& MeshEditorMode)
 
 		MeshEditorMode.TrackUndo(Mesh, Mesh->MakeUndo());
 	}
-#endif
 }
 
 #undef LOCTEXT_NAMESPACE
