Fixed local variable shadowing warning (Linux/Mac) from cl 4922582.

diff --git a/Engine/Plugins/Editor/MeshEditor/Source/PolygonModeling/UnifyNormals.cpp b/Engine/Plugins/Editor/MeshEditor/Source/PolygonModeling/UnifyNormals.cpp
--- a/Engine/Plugins/Editor/MeshEditor/Source/PolygonModeling/UnifyNormals.cpp
+++ b/Engine/Plugins/Editor/MeshEditor/Source/PolygonModeling/UnifyNormals.cpp
@@ -87,7 +87,6 @@ void UUnifyNormalsCommand::Execute(IMeshEditorModeEditingContract& MeshEditorMod
 		TSet<FPolygonID> CheckedPolygons;		// set of polygons already checked
 		TSet<FPolygonID> SelectedPolygons;		// set of selected polygons
 		TSet<FPolygonID> FlippedPolygons;		// set of polygons that are flipped with respect to the selected polygons
-		TSet<FPolygonID> BoundaryPolygons;		// set of polygons that represent the boundary of the flipped polygons
 		TSet<FVertexID> BoundaryPolygonVertices;// set of vertices of the polygons at the boundary of the flipped polygons
 		TArray<FVertexID> Vertices;
 
