CIS Fix #rb None
#rnx

diff --git a/Engine/Plugins/Editor/MeshEditor/Source/MeshEditor/MeshEditorMode.cpp b/Engine/Plugins/Editor/MeshEditor/Source/MeshEditor/MeshEditorMode.cpp
--- a/Engine/Plugins/Editor/MeshEditor/Source/MeshEditor/MeshEditorMode.cpp
+++ b/Engine/Plugins/Editor/MeshEditor/Source/MeshEditor/MeshEditorMode.cpp
@@ -967,8 +967,7 @@ void FMeshEditorMode::Enter()
 			TSharedPtr<IAssetViewport> ActiveLevelViewport = LevelEditor->GetActiveViewportInterface();
 			if (ActiveLevelViewport.IsValid())
 			{
-				const TSharedRef<SLevelViewport> LevelViewportWidget = StaticCastSharedRef<SLevelViewport>(ActiveLevelViewport->AsWidget());
-				ViewportClient = MakeShareable(&LevelViewportWidget->GetViewportClient());
+				ViewportClient = StaticCastSharedRef<SLevelViewport>(ActiveLevelViewport->AsWidget())->GetViewportClient();
 			}
 		}
 
