Hide the MeshModelingToolset plugin in the Plugin browser because it is automatically enabled by the parent Mode and not currently something we support using independently.
Hide MeshEditor plugin in the Plugin browser because it is deprecated/disabled and users should not see it.
#rnx

diff --git a/Engine/Plugins/Editor/MeshEditor/MeshEditor.uplugin b/Engine/Plugins/Editor/MeshEditor/MeshEditor.uplugin
--- a/Engine/Plugins/Editor/MeshEditor/MeshEditor.uplugin
+++ b/Engine/Plugins/Editor/MeshEditor/MeshEditor.uplugin
@@ -1,30 +1,30 @@
 {
-	"FileVersion" : 3,
-	"Version" : 1,
-	"VersionName" : "1.0",
-	"FriendlyName" : "Mesh Editor",
-	"Description" : "Mesh editing system",
-	"Category" : "Editor",
-	"CreatedBy" : "Epic Games, Inc.",
-	"CreatedByURL" : "http://epicgames.com",
-	"DocsURL" : "",
-	"MarketplaceURL" : "",
-	"SupportURL" : "",
-	"EnabledByDefault" : false,
-	"CanContainContent" : true,
-	"IsBetaVersion" : false,
-	"Installed" : false,
-	"Modules" :
-	[
+	"FileVersion": 3,
+	"Version": 1,
+	"VersionName": "1.0",
+	"FriendlyName": "Mesh Editor",
+	"Description": "Mesh editing system",
+	"Category": "Editor",
+	"CreatedBy": "Epic Games, Inc.",
+	"CreatedByURL": "http://epicgames.com",
+	"DocsURL": "",
+	"MarketplaceURL": "",
+	"SupportURL": "",
+	"EnabledByDefault": false,
+	"CanContainContent": true,
+	"IsBetaVersion": false,
+	"Installed": false,
+	"Hidden": true,
+	"Modules": [
 		{
-			"Name" : "MeshEditor",
-			"Type" : "Editor",
-			"LoadingPhase" : "Default"
+			"Name": "MeshEditor",
+			"Type": "Editor",
+			"LoadingPhase": "Default"
 		},
 		{
-			"Name" : "PolygonModeling",
-			"Type" : "Editor",
-			"LoadingPhase" : "Default"
+			"Name": "PolygonModeling",
+			"Type": "Editor",
+			"LoadingPhase": "Default"
 		}
 	],
 	"Plugins": [
@@ -38,7 +38,7 @@
 		},
 		{
 			"Name": "GeometryCollectionPlugin",
-			"Enabled":  true
-		}		
+			"Enabled": true
+		}
 	]
 }
diff --git a/Engine/Plugins/Experimental/MeshModelingToolset/MeshModelingToolset.uplugin b/Engine/Plugins/Experimental/MeshModelingToolset/MeshModelingToolset.uplugin
--- a/Engine/Plugins/Experimental/MeshModelingToolset/MeshModelingToolset.uplugin
+++ b/Engine/Plugins/Experimental/MeshModelingToolset/MeshModelingToolset.uplugin
@@ -13,6 +13,7 @@
 	"CanContainContent": true,
 	"IsBetaVersion": true,
 	"Installed": false,
+	"Hidden": true,
 	"Modules": [
 		{
 			"Name": "ModelingComponents",
@@ -43,7 +44,7 @@
 			"Type": "Editor",
 			"LoadingPhase": "Default",
 			"WhitelistPlatforms": [ "Win64" ]
-		},
+		}
 	],
 	"Plugins": [
 		{
@@ -52,7 +53,7 @@
 		},
 		{
 			"Name": "ProxyLODPlugin",
-			"Enabled" : true
+			"Enabled": true
 		}
 	]
 }
