Removing unneeded function call

diff --git a/Engine/Plugins/Editor/MeshEditor/Source/MeshEditor/MeshEditorMode.cpp b/Engine/Plugins/Editor/MeshEditor/Source/MeshEditor/MeshEditorMode.cpp
--- a/Engine/Plugins/Editor/MeshEditor/Source/MeshEditor/MeshEditorMode.cpp
+++ b/Engine/Plugins/Editor/MeshEditor/Source/MeshEditor/MeshEditorMode.cpp
@@ -979,7 +979,6 @@ void FMeshEditorMode::Exit()
 
 	if( VREditorMode && VREditorMode->IsFullyInitialized() )
 	{
-		VREditorMode->ResetActionsMenuGenerator();
 		VREditorMode->OnPlaceDraggedMaterial().RemoveAll( this );
 	}
 
