Removing VR delegate bindings

diff --git a/Engine/Plugins/Editor/MeshEditor/Source/MeshEditor/MeshEditorMode.cpp b/Engine/Plugins/Editor/MeshEditor/Source/MeshEditor/MeshEditorMode.cpp
--- a/Engine/Plugins/Editor/MeshEditor/Source/MeshEditor/MeshEditorMode.cpp
+++ b/Engine/Plugins/Editor/MeshEditor/Source/MeshEditor/MeshEditorMode.cpp
@@ -30,7 +30,6 @@
 #include "DrawDebugHelpers.h"	// For DrawDebugSphere
 #include "MeshEditorSettings.h"
 #include "LevelEditor.h"
-#include "VREditorMode.h" // @TODO: remove once radial menu code is in a non-VREditor module
 #include "Framework/MultiBox/MultiBoxBuilder.h"
 #include "Framework/Commands/UICommandList.h"
 #include "MeshEditorAssetContainer.h"
@@ -977,11 +976,6 @@ void FMeshEditorMode::Exit()
 	FFractureToolDelegates::Get().OnFractureExpansionBegin.RemoveAll(this);
 	FFractureToolDelegates::Get().OnFractureExpansionEnd.RemoveAll(this);
 
-	if( VREditorMode && VREditorMode->IsFullyInitialized() )
-	{
-		VREditorMode->OnPlaceDraggedMaterial().RemoveAll( this );
-	}
-
 	// If anything is selected, go ahead and deselect everything now
 	if( SelectedMeshElements.Num() > 0 )
 	{
@@ -2893,27 +2887,6 @@ void FMeshEditorMode::OnViewportInteractionFinishedMovingTransformables()
 }
 
 
-void FMeshEditorMode::OnVREditorModePlaceDraggedMaterial( UPrimitiveComponent* HitComponent, UMaterialInterface* MaterialInterface, bool& bPlaced )
-{
-	if( !bPlaced )
-	{
-		static TMap< UEditableMesh*, TArray<FMeshElement> > MeshesAndPolygons;
-		GetSelectedMeshesAndPolygons( /* Out */ MeshesAndPolygons );
-		for( const auto& MeshAndPolygons : MeshesAndPolygons )
-		{
-			if( MeshAndPolygons.Value.Num() > 0 )
-			{
-				if( MeshAndPolygons.Value[ 0 ].Component.Get() == HitComponent )
-				{
-					AssignMaterialToSelectedPolygons( MaterialInterface );
-					bPlaced = true;
-					break;
-				}
-			}
-		}
-	}
-}
-
 
 void FMeshEditorMode::UpdateActiveAction( const bool bIsActionFinishing )
 {
@@ -4575,125 +4548,6 @@ void FMeshEditorMode::OnActorSelectionChanged( const TArray<UObject*>& NewSelect
 	UpdateSelectedEditableMeshes();
 }
 
-void FMeshEditorMode::MakeVRRadialMenuActionsMenu(FMenuBuilder& MenuBuilder, TSharedPtr<FUICommandList> CommandList, class UVREditorMode* VRMode, float& RadiusOverride)
-{
-#if EDITABLE_MESH_USE_OPENSUBDIV
-	MenuBuilder.AddMenuEntry(
-		LOCTEXT("AddSubdivision", "Add SubD"),
-		FText(),
-		FSlateIcon(FMeshEditorStyle::GetStyleSetName(), "MeshEditorMode.AddSubdivision"),
-		FUIAction
-		(
-			FExecuteAction::CreateSP(this, &FMeshEditorMode::AddOrRemoveSubdivisionLevel, true )
-		),
-		NAME_None,
-		EUserInterfaceActionType::ToggleButton
-		);
-	MenuBuilder.AddMenuEntry(
-		LOCTEXT("RemoveSubdivision", "Remove SubD"),
-		FText(),
-		FSlateIcon(FMeshEditorStyle::GetStyleSetName(), "MeshEditorMode.RemoveSubdivision"),
-		FUIAction
-		(
-			FExecuteAction::CreateSP(this, &FMeshEditorMode::AddOrRemoveSubdivisionLevel, false)
-		),
-		NAME_None,
-		EUserInterfaceActionType::ToggleButton
-		);
-#endif
-	MenuBuilder.AddMenuEntry(
-		LOCTEXT("EditInstance", "Edit Instance"),
-		FText(),
-		FSlateIcon(FMeshEditorStyle::GetStyleSetName(), "MeshEditorMode.EditInstance"),
-		FUIAction
-		(
-			FExecuteAction::CreateLambda([this] { SetEditingPerInstance(IsEditingPerInstance() == false); }),
-			FCanExecuteAction::CreateLambda([this] {return true; }),
-			FIsActionChecked::CreateLambda([this] { return IsEditingPerInstance(); })
-		),
-		NAME_None,
-		EUserInterfaceActionType::ToggleButton
-		);
-
-	if (GetMeshElementSelectionMode() == EEditableMeshElementType::Polygon)
-	{
-		MenuBuilder.AddMenuEntry(
-			LOCTEXT("Move", "Move"),
-			FText(),
-			FSlateIcon(FMeshEditorStyle::GetStyleSetName(), "MeshEditorMode.PolyMove"),
-			FUIAction
-			(
-				FExecuteAction::CreateLambda([this] { SetEquippedAction(EEditableMeshElementType::Polygon, EMeshEditAction::Move ); }),
-				FCanExecuteAction::CreateSP(this, &FMeshEditorMode::IsMeshElementTypeSelectedOrIsActiveSelectionMode, EEditableMeshElementType::Polygon),
-				FIsActionChecked::CreateLambda([this] { return (EquippedPolygonAction == EMeshEditAction::Move); })
-				),
-			NAME_None,
-			EUserInterfaceActionType::ToggleButton
-			);
-	}
-	else if (GetMeshElementSelectionMode() == EEditableMeshElementType::Edge)
-	{
-		MenuBuilder.AddMenuEntry(
-			LOCTEXT("Move", "Move"),
-			FText(),
-			FSlateIcon(FMeshEditorStyle::GetStyleSetName(), "MeshEditorMode.EdgeMove"),
-			FUIAction
-			(
-				FExecuteAction::CreateLambda([this] { SetEquippedAction( EEditableMeshElementType::Edge, EMeshEditAction::Move ); }),
-				FCanExecuteAction::CreateSP(this, &FMeshEditorMode::IsMeshElementTypeSelectedOrIsActiveSelectionMode, EEditableMeshElementType::Edge),
-				FIsActionChecked::CreateLambda([this] { return (EquippedEdgeAction == EMeshEditAction::Move); })
-			),
-			NAME_None,
-			EUserInterfaceActionType::ToggleButton
-			);
-		MenuBuilder.AddMenuEntry(
-			LOCTEXT( "SelectEdgeLoop", "Select Edge Loop" ),
-			FText(),
-			FSlateIcon(FMeshEditorStyle::GetStyleSetName(), "MeshEditorMode.SelectLoop"),
-			FUIAction
-			(
-				FExecuteAction::CreateLambda([this] { SelectEdgeLoops(); }),
-				FCanExecuteAction::CreateSP(this, &FMeshEditorMode::IsMeshElementTypeSelected, EEditableMeshElementType::Edge)
-				),
-			NAME_None,
-			EUserInterfaceActionType::CollapsedButton
-			);
-	}
-	else if (GetMeshElementSelectionMode() == EEditableMeshElementType::Vertex)
-	{
-		MenuBuilder.AddMenuEntry(
-			LOCTEXT("Move", "Move"),
-			FText(),
-			FSlateIcon(FMeshEditorStyle::GetStyleSetName(), "MeshEditorMode.VertexMove"),
-			FUIAction
-			(
-				FExecuteAction::CreateLambda([this] { SetEquippedAction( EEditableMeshElementType::Vertex, EMeshEditAction::Move ); }),
-				FCanExecuteAction::CreateSP(this, &FMeshEditorMode::IsMeshElementTypeSelectedOrIsActiveSelectionMode, EEditableMeshElementType::Vertex),
-				FIsActionChecked::CreateLambda([this] { return (EquippedVertexAction == EMeshEditAction::Move); })
-				),
-			NAME_None,
-			EUserInterfaceActionType::ToggleButton
-			);
-		MenuBuilder.AddMenuEntry(
-			LOCTEXT("WeldSelected", "Weld Selected"),
-			FText(),
-			FSlateIcon(FMeshEditorStyle::GetStyleSetName(), "MeshEditorMode.VertexWeld"),
-			FUIAction
-			(
-				FExecuteAction::CreateLambda([this] { WeldSelectedVertices(); }),
-				FCanExecuteAction::CreateSP(this, &FMeshEditorMode::IsMeshElementTypeSelected, EEditableMeshElementType::Vertex)
-				),
-			NAME_None,
-			EUserInterfaceActionType::CollapsedButton
-			);
-	}
-
-	for( UMeshEditorCommand* Command : MeshEditorCommands::Get() )
-	{
-		Command->AddToVRRadialMenuActionsMenu( *this, MenuBuilder, CommandList, FMeshEditorStyle::GetStyleSetName(), VRMode );
-	}
-}
-
 
 FName FMeshEditorMode::GetEquippedAction( const EEditableMeshElementType ForElementType ) const
 {
diff --git a/Engine/Plugins/Editor/MeshEditor/Source/MeshEditor/MeshEditorMode.h b/Engine/Plugins/Editor/MeshEditor/Source/MeshEditor/MeshEditorMode.h
--- a/Engine/Plugins/Editor/MeshEditor/Source/MeshEditor/MeshEditorMode.h
+++ b/Engine/Plugins/Editor/MeshEditor/Source/MeshEditor/MeshEditorMode.h
@@ -165,9 +165,6 @@ protected:
 
 	/** Called when the viewport interaction finishes moving a set of transformable objects */
 	void OnViewportInteractionFinishedMovingTransformables();
-	
-	/** Called when VR editor world interaction drags a material onto a component */
-	void OnVREditorModePlaceDraggedMaterial( class UPrimitiveComponent* HitComponent, class UMaterialInterface* MaterialInterface, bool& bPlaced );
 
 	/** Applies a modification to the mesh that's currently hovered */
 	void UpdateActiveAction( const bool bIsActionFinishing );
@@ -273,9 +270,6 @@ protected:
 	/** Callback from the level editor when new actors become selected or deselected */
 	void OnActorSelectionChanged( const TArray<UObject*>& NewSelection, bool bForceRefresh );
 
-	/** Creates the mesh edit actions to pass to the radial menu generator */
-	void MakeVRRadialMenuActionsMenu(class FMenuBuilder& MenuBuilder, TSharedPtr<class FUICommandList> CommandList, class UVREditorMode* VRMode, float& RadiusOverride);
-
 	/** Clears any references to editable meshes (which may now be invalid) */
 	void RemoveEditableMeshReferences();
 
@@ -439,9 +433,6 @@ protected:
 	/** Cached pointer to the viewport world interaction object we're using to interact with mesh elements */
 	class UViewportWorldInteraction* ViewportWorldInteraction;
 
-	/** Cached pointer to the VREditorMode object we're using */
-	class UVREditorMode* VREditorMode;
-
 	/** Material to use to render hovered mesh geometry */
 	UMaterialInterface* HoveredGeometryMaterial;
 
